Separating BlendStack from PoseSearch into its own plugin

diff --git a/Engine/Plugins/Experimental/Animation/BlendStack/BlendStack.uplugin b/Engine/Plugins/Experimental/Animation/BlendStack/BlendStack.uplugin
new file mode 100644
--- /dev/null
+++ b/Engine/Plugins/Experimental/Animation/BlendStack/BlendStack.uplugin
@@ -0,0 +1,36 @@
+{
+	"FileVersion": 3,
+	"Version": 1,
+	"VersionName": "0.1",
+	"FriendlyName": "Blend Stack",
+	"Description": "Experimental blend stack API",
+	"Category": "Animation",
+	"CreatedBy": "Epic Games, Inc.",
+	"CreatedByURL": "https://epicgames.com",
+	"DocsURL": "",
+	"MarketplaceURL": "",
+	"SupportURL": "",
+	"EnabledByDefault": false,
+	"CanContainContent": false,
+	"IsBetaVersion": false,
+	"IsExperimentalVersion": true,
+	"Installed": false,
+	"Modules": [
+		{
+			"Name": "BlendStack",
+			"Type": "Runtime",
+			"LoadingPhase": "PreDefault"
+		},
+		{
+			"Name": "BlendStackEditor",
+			"Type": "UncookedOnly",
+			"LoadingPhase": "PreDefault"
+		}
+	],
+	"Plugins": [
+		{
+			"Name": "StructUtils",
+			"Enabled": true
+		}
+	]
+}
\ No newline at end of file
diff --git a/Engine/Plugins/Experimental/Animation/BlendStack/Config/BaseBlendStack.ini b/Engine/Plugins/Experimental/Animation/BlendStack/Config/BaseBlendStack.ini
new file mode 100644
--- /dev/null
+++ b/Engine/Plugins/Experimental/Animation/BlendStack/Config/BaseBlendStack.ini
@@ -0,0 +1,10 @@
+[CoreRedirects]
++ClassRedirects=(OldName="/Script/PoseSearchEditor.AnimationBlendStackGraph",NewName="/Script/BlendStackEditor.AnimationBlendStackGraph")
++ClassRedirects=(OldName="/Script/PoseSearchEditor.AnimationBlendStackGraphSchema",NewName="/Script/BlendStackEditor.AnimationBlendStackGraphSchema")
++ClassRedirects=(OldName="/Script/PoseSearchEditor.AnimGraphNode_BlendStack",NewName="/Script/BlendStackEditor.AnimGraphNode_BlendStack")
++ClassRedirects=(OldName="/Script/PoseSearchEditor.AnimGraphNode_BlendStackInput",NewName="/Script/BlendStackEditor.AnimGraphNode_BlendStackInput")
++ClassRedirects=(OldName="/Script/PoseSearchEditor.AnimGraphNode_BlendStackResult",NewName="/Script/BlendStackEditor.AnimGraphNode_BlendStackResult")
++ClassRedirects=(OldName="/Script/PoseSearch.BlendStackAnimNodeLibrary",NewName="/Script/BlendStack.BlendStackAnimNodeLibrary")
++StructRedirects=(OldName="/Script/PoseSearch.AnimNode_BlendStack",NewName="/Script/BlendStack.AnimNode_BlendStack")
++StructRedirects=(OldName="/Script/PoseSearch.AnimNode_BlendStackInput",NewName="/Script/BlendStack.AnimNode_BlendStackInput")
++StructRedirects=(OldName="/Script/PoseSearch.AnimNode_BlendStackResult",NewName="/Script/BlendStack.AnimNode_BlendStackResult")
diff --git a/Engine/Plugins/Experimental/Animation/BlendStack/Source/Editor/BlendStackEditor.Build.cs b/Engine/Plugins/Experimental/Animation/BlendStack/Source/Editor/BlendStackEditor.Build.cs
new file mode 100644
--- /dev/null
+++ b/Engine/Plugins/Experimental/Animation/BlendStack/Source/Editor/BlendStackEditor.Build.cs
@@ -0,0 +1,33 @@
+// Copyright Epic Games, Inc. All Rights Reserved.
+
+using System.IO;
+using UnrealBuildTool;
+
+public class BlendStackEditor : ModuleRules
+{
+	public BlendStackEditor(ReadOnlyTargetRules Target) : base(Target)
+	{
+		PCHUsage = ModuleRules.PCHUsageMode.UseExplicitOrSharedPCHs;
+
+		PublicDependencyModuleNames.AddRange(
+			new string[]
+			{
+				"AnimGraph",
+				"AnimGraphRuntime",
+				"AnimationCore",
+				"Core",
+				"CoreUObject",
+				"Engine",
+				"BlendStack",
+			}
+		);
+
+		PrivateDependencyModuleNames.AddRange(
+			new string[]
+			{
+				"BlueprintGraph",
+				"UnrealEd"
+			}
+		);
+	}
+}
diff --git a/Engine/Plugins/Experimental/Animation/BlendStack/Source/Editor/Private/AnimGraphNode_BlendStack.cpp b/Engine/Plugins/Experimental/Animation/BlendStack/Source/Editor/Private/AnimGraphNode_BlendStack.cpp
new file mode 100644
--- /dev/null
+++ b/Engine/Plugins/Experimental/Animation/BlendStack/Source/Editor/Private/AnimGraphNode_BlendStack.cpp
@@ -0,0 +1,295 @@
+// Copyright Epic Games, Inc. All Rights Reserved.
+
+#include "AnimGraphNode_BlendStack.h"
+#include "Animation/AnimRootMotionProvider.h"
+#include "AnimationBlendStackGraphSchema.h"
+#include "AnimationBlendStackGraph.h"
+#include "Kismet2/BlueprintEditorUtils.h"
+#include "Kismet2/KismetEditorUtilities.h"
+#include "EdGraphUtilities.h"
+#include "IAnimBlueprintCompilationContext.h"
+#include "AnimGraphNode_Root.h"
+#include "AnimGraphNode_BlendStackInput.h"
+
+#define LOCTEXT_NAMESPACE "AnimGraphNode_BlendStack"
+
+
+FLinearColor UAnimGraphNode_BlendStack::GetNodeTitleColor() const
+{
+	return FColor(86, 182, 194);
+}
+
+FText UAnimGraphNode_BlendStack::GetTooltipText() const
+{
+	return LOCTEXT("NodeToolTip", "Blend Stack");
+}
+
+FText UAnimGraphNode_BlendStack::GetNodeTitle(ENodeTitleType::Type TitleType) const
+{
+	return LOCTEXT("NodeTitle", "Blend Stack");
+}
+
+FText UAnimGraphNode_BlendStack::GetMenuCategory() const
+{
+	return LOCTEXT("NodeCategory", "Animation|Blend Stack");
+}
+
+void UAnimGraphNode_BlendStack_Base::GetOutputLinkAttributes(FNodeAttributeArray& OutAttributes) const
+{
+	if (UE::Anim::IAnimRootMotionProvider::Get())
+	{
+		OutAttributes.Add(UE::Anim::IAnimRootMotionProvider::AttributeName);
+	}
+}
+
+void UAnimGraphNode_BlendStack::BakeDataDuringCompilation(class FCompilerResultsLog& MessageLog)
+{
+	UAnimBlueprint* AnimBlueprint = GetAnimBlueprint();
+	AnimBlueprint->FindOrAddGroup(Node.GetGroupName());
+}
+
+void UAnimGraphNode_BlendStack_Base::ExpandGraphAndProcessNodes(
+	UEdGraph* SourceGraph, 
+	UAnimGraphNode_Base* SourceRootNode, TArrayView<UAnimGraphNode_BlendStackInput*> InputNodes,
+	IAnimBlueprintCompilationContext& InCompilationContext, IAnimBlueprintGeneratedClassCompiledData& OutCompiledData,
+	UAnimGraphNode_Base*& OutRootNode, TArrayView<UAnimGraphNode_BlendStackInput*> OutInputNodes)
+{
+	// Note: This is mostly copied from UAnimGraphNode_BlendSpaceGraphBase::ExpandGraphAndProcessNodes
+	
+	// Clone the nodes from the source graph
+	// Note that we outer this graph to the ConsolidatedEventGraph to allow ExpansionStep to 
+	// correctly retrieve the context for any expanded function calls (custom make/break structs etc.)
+	UEdGraph* ClonedGraph = FEdGraphUtilities::CloneGraph(SourceGraph, InCompilationContext.GetConsolidatedEventGraph(), &InCompilationContext.GetMessageLog(), true);
+
+	// Grab all the animation nodes and find the corresponding 
+	// root node and input pose nodes in the cloned set
+	TArray<UAnimGraphNode_Base*> AnimNodeList;
+
+	const UObject* SourceRootObject = InCompilationContext.GetMessageLog().FindSourceObject(SourceRootNode);
+	TArray<UObject*> SourceInputObjects;
+	SourceInputObjects.SetNum(InputNodes.Num());
+	for (int32 Index = 0; Index < InputNodes.Num(); ++Index)
+	{
+		SourceInputObjects[Index] = InCompilationContext.GetMessageLog().FindSourceObject(InputNodes[Index]);
+	}
+
+	for (auto NodeIt = ClonedGraph->Nodes.CreateIterator(); NodeIt; ++NodeIt)
+	{
+		UEdGraphNode* ClonedNode = *NodeIt;
+		if (UAnimGraphNode_Base* TestNode = Cast<UAnimGraphNode_Base>(ClonedNode))
+		{
+			AnimNodeList.Add(TestNode);
+
+			//@TODO: There ought to be a better way to determine this
+			UObject* TestObject = InCompilationContext.GetMessageLog().FindSourceObject(TestNode);
+			if (TestObject == SourceRootObject)
+			{
+				OutRootNode = TestNode;
+			}
+
+			int32 FoundIndex;
+			if (SourceInputObjects.Find(TestObject, FoundIndex))
+			{
+				OutInputNodes[FoundIndex] = (UAnimGraphNode_BlendStackInput*)TestNode;
+			}
+		}
+	}
+
+	check(OutRootNode && !OutInputNodes.Contains(nullptr));
+
+	// Run another expansion pass to catch the graph we just added (this is slightly wasteful)
+	InCompilationContext.ExpansionStep(ClonedGraph, false);
+
+	// Validate graph now we have expanded/pruned
+	InCompilationContext.ValidateGraphIsWellFormed(ClonedGraph);
+
+	// Move the cloned nodes into the consolidated event graph
+	const bool bIsLoading = InCompilationContext.GetBlueprint()->bIsRegeneratingOnLoad || IsAsyncLoading();
+	const bool bIsCompiling = InCompilationContext.GetBlueprint()->bBeingCompiled;
+	ClonedGraph->MoveNodesToAnotherGraph(InCompilationContext.GetConsolidatedEventGraph(), bIsLoading, bIsCompiling);
+
+	// Process any animation nodes
+	{
+		TArray<UAnimGraphNode_Base*> RootSet;
+		RootSet.Add(OutRootNode);
+
+		InCompilationContext.PruneIsolatedAnimationNodes(RootSet, AnimNodeList);
+		InCompilationContext.ProcessAnimationNodes(AnimNodeList);
+	}
+}
+
+void UAnimGraphNode_BlendStack_Base::OnProcessDuringCompilation(IAnimBlueprintCompilationContext& InCompilationContext, IAnimBlueprintGeneratedClassCompiledData& OutCompiledData)
+{
+	Super::OnProcessDuringCompilation(InCompilationContext, OutCompiledData);
+
+	FAnimNode_BlendStack_Standalone* AnimNode = GetBlendStackNode();
+	check(AnimNode);
+
+	const int32 MaxBlendsNum = GetMaxActiveBlends();
+	// Set MaxActiveBlends of the blend stack at compile time.
+	// @todo: Allow reducing the blend stack size (but not decreasing), to enable scalability control under i.e. lower LODs
+	AnimNode->MaxActiveBlends = GetMaxActiveBlends();
+
+	UAnimationBlendStackGraph* SampleGraph = CastChecked<UAnimationBlendStackGraph>(BoundGraph);
+	if (!SampleGraph->ResultNode->IsNodeRootSet())
+	{
+		// Input Pose is connected to Output Pose, so the sample graph does nothing.
+		// No need to use allocate sample graphs in that case.
+		AnimNode->SampleGraphPoseLinks.Reset();
+		return;
+	}
+
+	// Allocate one sample graph per-active blend plus an extra one for the stored pose.
+	AnimNode->SampleGraphPoseLinks.SetNum(MaxBlendsNum + 1);
+
+	TArray<UAnimGraphNode_BlendStackInput*> InputNodes;
+	BoundGraph->GetNodesOfClass<UAnimGraphNode_BlendStackInput>(InputNodes);
+
+	int32 BlendStackAllocationIndex = InCompilationContext.GetAllocationIndexOfNode(this);
+	for(int32 Index = 0; Index < AnimNode->SampleGraphPoseLinks.Num(); ++ Index)
+	{
+		FBlendStack_SampleGraphPoseLink& GraphPoseLink = AnimNode->SampleGraphPoseLinks[Index];
+		UAnimGraphNode_Base* ClonedRootNode;
+		TArray<UAnimGraphNode_BlendStackInput*> ClonedInputNodes;
+		ClonedInputNodes.SetNum(InputNodes.Num());
+
+		ExpandGraphAndProcessNodes(SampleGraph, SampleGraph->ResultNode, InputNodes, InCompilationContext, OutCompiledData, ClonedRootNode, ClonedInputNodes);
+
+		// Blend stack node is potentially nested in the struct, so we can't use FPoseLinkMappingRecord. Patch at runtime instead.
+		for (UAnimGraphNode_BlendStackInput* InputNode : ClonedInputNodes)
+		{
+			InputNode->Node.BlendStackAllocationIndex = BlendStackAllocationIndex;
+			InputNode->Node.SampleIndex = Index;
+		}
+		GraphPoseLink.RootNodeIndex = InCompilationContext.GetAllocationIndexOfNode(ClonedRootNode);
+	}
+}
+
+bool UAnimGraphNode_BlendStack_Base::DoesSupportTimeForTransitionGetter() const
+{
+	return true;
+}
+
+UAnimationAsset* UAnimGraphNode_BlendStack_Base::GetAnimationAsset() const
+{
+	return nullptr;
+}
+
+const TCHAR* UAnimGraphNode_BlendStack_Base::GetTimePropertyName() const
+{
+	return TEXT("InternalTimeAccumulator");
+}
+
+UScriptStruct* UAnimGraphNode_BlendStack_Base::GetTimePropertyStruct() const
+{
+	return FAnimNode_BlendStack::StaticStruct();
+}
+
+void UAnimGraphNode_BlendStack_Base::Serialize(FArchive& Ar)
+{
+	Super::Serialize(Ar);
+
+	if (Ar.IsLoading() && (BoundGraph == nullptr))
+	{
+		// Update older nodes without a graph to create one on load.
+		CreateGraph();
+	}
+}
+
+int32 UAnimGraphNode_BlendStack_Base::GetMaxActiveBlends() const
+{
+	FAnimNode_BlendStack_Standalone* AnimNode = GetBlendStackNode();
+	check(AnimNode);
+	return AnimNode->MaxActiveBlends;
+}
+
+void UAnimGraphNode_BlendStack_Base::CreateGraph()
+{
+	// Create a new animation graph
+	check(BoundGraph == nullptr);
+	BoundGraph = FBlueprintEditorUtils::CreateNewGraph(
+		this,
+		NAME_None,
+		UAnimationBlendStackGraph::StaticClass(),
+		UAnimationBlendStackGraphSchema::StaticClass());
+	check(BoundGraph);
+
+	// Initialize the anim graph
+	const UEdGraphSchema* Schema = BoundGraph->GetSchema();
+	Schema->CreateDefaultNodesForGraph(*BoundGraph);
+
+	// Add the new graph as a child of our parent graph
+	UEdGraph* ParentGraph = GetGraph();
+	if(ParentGraph->SubGraphs.Find(BoundGraph) == INDEX_NONE)
+	{
+		ParentGraph->Modify();
+		ParentGraph->SubGraphs.Add(BoundGraph);
+	}
+}
+
+void UAnimGraphNode_BlendStack_Base::PostPlacedNewNode()
+{
+	Super::PostPlacedNewNode();
+	CreateGraph();
+}
+
+TArray<UEdGraph*> UAnimGraphNode_BlendStack_Base::GetSubGraphs() const
+{
+	return TArray<UEdGraph*>({BoundGraph});
+}
+
+UObject* UAnimGraphNode_BlendStack_Base::GetJumpTargetForDoubleClick() const
+{
+	// Open the blend stack graph
+	return BoundGraph;
+}
+
+void UAnimGraphNode_BlendStack_Base::JumpToDefinition() const
+{
+	if (UObject* HyperlinkTarget = GetJumpTargetForDoubleClick())
+	{
+		FKismetEditorUtilities::BringKismetToFocusAttentionOnObject(HyperlinkTarget);
+	}
+}
+
+void UAnimGraphNode_BlendStack_Base::DestroyNode()
+{
+	UBlueprint* Blueprint = GetBlueprint();
+	UEdGraph* GraphToRemove = BoundGraph;
+	BoundGraph = nullptr;
+	Super::DestroyNode();
+
+	if (GraphToRemove)
+	{
+		GraphToRemove->Modify();
+		FBlueprintEditorUtils::RemoveGraph(Blueprint, GraphToRemove, EGraphRemoveFlags::Recompile);
+	}
+}
+
+void UAnimGraphNode_BlendStack_Base::PostPasteNode()
+{
+	Super::PostPasteNode();
+
+	if(BoundGraph)
+	{
+		// Add the new graph as a child of our parent graph
+		UEdGraph* ParentGraph = GetGraph();
+
+		if(ParentGraph->SubGraphs.Find(BoundGraph) == INDEX_NONE)
+		{
+			ParentGraph->SubGraphs.Add(BoundGraph);
+		}
+
+		for (UEdGraphNode* GraphNode : BoundGraph->Nodes)
+		{
+			GraphNode->CreateNewGuid();
+			GraphNode->PostPasteNode();
+			GraphNode->ReconstructNode();
+		}
+
+		//restore transactional flag that is lost during copy/paste process
+		BoundGraph->SetFlags(RF_Transactional);
+	}
+}
+
+#undef LOCTEXT_NAMESPACE
diff --git a/Engine/Plugins/Experimental/Animation/BlendStack/Source/Editor/Private/AnimGraphNode_BlendStackInput.cpp b/Engine/Plugins/Experimental/Animation/BlendStack/Source/Editor/Private/AnimGraphNode_BlendStackInput.cpp
new file mode 100644
--- /dev/null
+++ b/Engine/Plugins/Experimental/Animation/BlendStack/Source/Editor/Private/AnimGraphNode_BlendStackInput.cpp
@@ -0,0 +1,35 @@
+// Copyright Epic Games, Inc. All Rights Reserved.
+
+#include "AnimGraphNode_BlendStackInput.h"
+#include "AnimGraphNode_BlendStack.h"
+#include "AnimationBlendStackGraph.h"
+
+#define LOCTEXT_NAMESPACE "BlendStackInput"
+
+FText UAnimGraphNode_BlendStackInput::GetNodeTitle(ENodeTitleType::Type TitleType) const
+{
+	FText DefaultTitle = LOCTEXT("Title", "Blend Stack Input");
+	return DefaultTitle;
+}
+
+FText UAnimGraphNode_BlendStackInput::GetMenuCategory() const
+{
+	return LOCTEXT("BlendStackInputCategory", "Animation|Blend Stack");
+}
+
+bool UAnimGraphNode_BlendStackInput::CanUserDeleteNode() const
+{
+	return true;
+}
+
+bool UAnimGraphNode_BlendStackInput::CanDuplicateNode() const
+{
+	return true;
+}
+
+bool UAnimGraphNode_BlendStackInput::IsCompatibleWithGraph(UEdGraph const* Graph) const
+{
+	return Graph->IsA(UAnimationBlendStackGraph::StaticClass());
+}
+
+#undef LOCTEXT_NAMESPACE
diff --git a/Engine/Plugins/Experimental/Animation/BlendStack/Source/Editor/Private/AnimGraphNode_BlendStackResult.cpp b/Engine/Plugins/Experimental/Animation/BlendStack/Source/Editor/Private/AnimGraphNode_BlendStackResult.cpp
new file mode 100644
--- /dev/null
+++ b/Engine/Plugins/Experimental/Animation/BlendStack/Source/Editor/Private/AnimGraphNode_BlendStackResult.cpp
@@ -0,0 +1,18 @@
+// Copyright Epic Games, Inc. All Rights Reserved.
+
+#include "AnimGraphNode_BlendStackResult.h"
+#include "AnimGraphNode_BlendStackInput.h"
+
+bool UAnimGraphNode_BlendStackResult::IsNodeRootSet() const
+{
+	if (!Pins.IsEmpty() && !Pins[0]->LinkedTo.IsEmpty())
+	{
+		if (Pins[0]->LinkedTo[0]->GetOuter()->IsA<UAnimGraphNode_BlendStackInput>())
+		{
+			// If output is connected to input directly, graph is a no-op. Consider it not set for pruning purposes.
+			return false;
+		}
+	}
+
+	return true;
+}
diff --git a/Engine/Plugins/Experimental/Animation/BlendStack/Source/Editor/Private/AnimationBlendStackGraphSchema.cpp b/Engine/Plugins/Experimental/Animation/BlendStack/Source/Editor/Private/AnimationBlendStackGraphSchema.cpp
new file mode 100644
--- /dev/null
+++ b/Engine/Plugins/Experimental/Animation/BlendStack/Source/Editor/Private/AnimationBlendStackGraphSchema.cpp
@@ -0,0 +1,40 @@
+// Copyright Epic Games, Inc. All Rights Reserved.
+
+#include "AnimationBlendStackGraphSchema.h"
+#include "AnimationBlendStackGraph.h"
+#include "AnimGraphNode_BlendStackInput.h"
+#include "AnimGraphNode_BlendStackResult.h"
+
+/////////////////////////////////////////////////////
+// UAnimationBlendStackGraphSchema
+
+void UAnimationBlendStackGraphSchema::CreateDefaultNodesForGraph(UEdGraph& Graph) const
+{
+	UAnimationBlendStackGraph* TypedGraph = CastChecked<UAnimationBlendStackGraph>(&Graph);
+
+	UAnimGraphNode_BlendStackResult* ResultNode;
+	UAnimGraphNode_BlendStackInput* InputNode;
+
+	{
+		// Create the result node
+		FGraphNodeCreator<UAnimGraphNode_BlendStackResult> NodeCreator(Graph);
+		ResultNode = NodeCreator.CreateNode();
+		NodeCreator.Finalize();
+		SetNodeMetaData(ResultNode, FNodeMetadata::DefaultGraphNode);
+		TypedGraph->ResultNode = ResultNode;
+	}
+
+	{
+		// Create the result node
+		FGraphNodeCreator<UAnimGraphNode_BlendStackInput> NodeCreator(Graph);
+		InputNode = NodeCreator.CreateNode();
+		NodeCreator.Finalize();
+		SetNodeMetaData(InputNode, FNodeMetadata::DefaultGraphNode);
+
+		// Move input node to the left of the output pose.
+		InputNode->NodePosX -= 200;
+	}
+
+	// Connect input node to output pose.
+	InputNode->Pins[0]->MakeLinkTo(ResultNode->Pins[0]);
+}
\ No newline at end of file
diff --git a/Engine/Plugins/Experimental/Animation/BlendStack/Source/Editor/Private/BlendStackEditor.cpp b/Engine/Plugins/Experimental/Animation/BlendStack/Source/Editor/Private/BlendStackEditor.cpp
new file mode 100644
--- /dev/null
+++ b/Engine/Plugins/Experimental/Animation/BlendStack/Source/Editor/Private/BlendStackEditor.cpp
@@ -0,0 +1,37 @@
+// Copyright Epic Games, Inc. All Rights Reserved.
+
+#include "BlendStackEditor.h"
+#include "Modules/ModuleManager.h"
+
+
+DEFINE_LOG_CATEGORY(LogBlendStackEditor);
+
+#define LOCTEXT_NAMESPACE "BlendStackEditorModule"
+
+namespace UE::BlendStack
+{
+
+//////////////////////////////////////////////////////////////////////////
+// Blend Stack Editor Module
+class FEditorModule : public IBlendStackEditorModuleInterface
+{
+public:
+	void StartupModule() override;
+	void ShutdownModule() override;
+};
+
+void FEditorModule::StartupModule()
+{
+}
+
+void FEditorModule::ShutdownModule()
+{
+
+}
+
+
+} // namespace UE::BlendStack
+
+IMPLEMENT_MODULE(UE::BlendStack::FEditorModule, BlendStackEditor);
+
+#undef LOCTEXT_NAMESPACE
diff --git a/Engine/Plugins/Experimental/Animation/BlendStack/Source/Editor/Public/AnimGraphNode_BlendStack.h b/Engine/Plugins/Experimental/Animation/BlendStack/Source/Editor/Public/AnimGraphNode_BlendStack.h
new file mode 100644
--- /dev/null
+++ b/Engine/Plugins/Experimental/Animation/BlendStack/Source/Editor/Public/AnimGraphNode_BlendStack.h
@@ -0,0 +1,68 @@
+// Copyright Epic Games, Inc. All Rights Reserved.
+
+#pragma once
+
+#include "AnimGraphNode_AssetPlayerBase.h"
+#include "BlendStack/AnimNode_BlendStack.h"
+#include "AnimGraphNode_BlendStack.generated.h"
+
+class UAnimGraphNode_BlendStackInput;
+
+UCLASS(Abstract, Experimental)
+class BLENDSTACKEDITOR_API UAnimGraphNode_BlendStack_Base : public UAnimGraphNode_AssetPlayerBase
+{
+	GENERATED_BODY()
+
+	virtual void GetOutputLinkAttributes(FNodeAttributeArray& OutAttributes) const override;
+	virtual void OnProcessDuringCompilation(IAnimBlueprintCompilationContext& InCompilationContext, IAnimBlueprintGeneratedClassCompiledData& OutCompiledData) override;
+
+	virtual bool DoesSupportTimeForTransitionGetter() const override;
+	virtual UAnimationAsset* GetAnimationAsset() const override;
+	virtual const TCHAR* GetTimePropertyName() const override;
+	virtual UScriptStruct* GetTimePropertyStruct() const override;
+
+	virtual void PostPlacedNewNode() override;
+	virtual UObject* GetJumpTargetForDoubleClick() const override;
+	virtual void JumpToDefinition() const override;
+	virtual void DestroyNode() override;
+	virtual void PostPasteNode() override;
+	virtual void Serialize(FArchive& Ar) override;
+
+	// Helper function for compilation
+	void ExpandGraphAndProcessNodes(
+		UEdGraph* SourceGraph, 
+		UAnimGraphNode_Base* SourceRootNode, TArrayView<UAnimGraphNode_BlendStackInput*> SourceInputNode,
+		IAnimBlueprintCompilationContext& InCompilationContext, IAnimBlueprintGeneratedClassCompiledData& OutCompiledData,
+		UAnimGraphNode_Base*& OutRootNode, TArrayView<UAnimGraphNode_BlendStackInput*> OutInputNodes);
+
+	virtual TArray<UEdGraph*> GetSubGraphs() const override;
+
+protected:
+	virtual FAnimNode_BlendStack_Standalone* GetBlendStackNode() const PURE_VIRTUAL(UAnimGraphNode_BlendStack_Base::GetBlendStackNode, return nullptr;);
+	int32 GetMaxActiveBlends() const;
+
+private:
+
+	void CreateGraph();
+
+	UPROPERTY()
+	TObjectPtr<UEdGraph> BoundGraph = nullptr;
+};
+
+UCLASS(MinimalAPI, Experimental)
+class UAnimGraphNode_BlendStack : public UAnimGraphNode_BlendStack_Base
+{
+	GENERATED_BODY()
+
+	UPROPERTY(EditAnywhere, Category = Settings)
+	FAnimNode_BlendStack Node;
+
+	virtual FLinearColor GetNodeTitleColor() const override;
+	virtual FText GetTooltipText() const override;
+	virtual FText GetNodeTitle(ENodeTitleType::Type TitleType) const override;
+	virtual FText GetMenuCategory() const override;
+	virtual void BakeDataDuringCompilation(class FCompilerResultsLog& MessageLog) override;
+
+protected:
+	virtual FAnimNode_BlendStack_Standalone* GetBlendStackNode() const override { return (FAnimNode_BlendStack_Standalone*)(&Node); }
+};
\ No newline at end of file
diff --git a/Engine/Plugins/Experimental/Animation/BlendStack/Source/Editor/Public/AnimGraphNode_BlendStackInput.h b/Engine/Plugins/Experimental/Animation/BlendStack/Source/Editor/Public/AnimGraphNode_BlendStackInput.h
new file mode 100644
--- /dev/null
+++ b/Engine/Plugins/Experimental/Animation/BlendStack/Source/Editor/Public/AnimGraphNode_BlendStackInput.h
@@ -0,0 +1,29 @@
+// Copyright Epic Games, Inc. All Rights Reserved.
+
+#pragma once
+
+#include "AnimGraphNode_Base.h"
+#include "BlendStack/AnimNode_BlendStackInput.h"
+
+#include "AnimGraphNode_BlendStackInput.generated.h"
+
+UCLASS(MinimalAPI)
+class UAnimGraphNode_BlendStackInput : public UAnimGraphNode_Base
+{
+	GENERATED_BODY()
+
+public:
+	UPROPERTY(EditAnywhere, Category = "Inputs")
+	FAnimNode_BlendStackInput Node;
+
+	virtual FText GetNodeTitle(ENodeTitleType::Type TitleType) const override;
+	virtual FText GetMenuCategory() const override;
+
+	virtual bool CanUserDeleteNode() const override;
+	virtual bool CanDuplicateNode() const override;
+	virtual bool IsCompatibleWithGraph(UEdGraph const* Graph) const override;
+
+	/** @return whether this node is editable. Non-editable nodes are implemented from function interfaces */
+	bool IsEditable() const { return CanUserDeleteNode(); }
+
+};
\ No newline at end of file
diff --git a/Engine/Plugins/Experimental/Animation/BlendStack/Source/Editor/Public/AnimGraphNode_BlendStackResult.h b/Engine/Plugins/Experimental/Animation/BlendStack/Source/Editor/Public/AnimGraphNode_BlendStackResult.h
new file mode 100644
--- /dev/null
+++ b/Engine/Plugins/Experimental/Animation/BlendStack/Source/Editor/Public/AnimGraphNode_BlendStackResult.h
@@ -0,0 +1,15 @@
+// Copyright Epic Games, Inc. All Rights Reserved.
+
+#pragma once
+
+#include "AnimGraphNode_Root.h"
+
+#include "AnimGraphNode_BlendStackResult.generated.h"
+
+UCLASS(MinimalAPI)
+class UAnimGraphNode_BlendStackResult : public UAnimGraphNode_Root
+{
+	GENERATED_BODY()
+
+	virtual bool IsNodeRootSet() const override;
+};
\ No newline at end of file
diff --git a/Engine/Plugins/Experimental/Animation/BlendStack/Source/Editor/Public/AnimationBlendStackGraph.h b/Engine/Plugins/Experimental/Animation/BlendStack/Source/Editor/Public/AnimationBlendStackGraph.h
new file mode 100644
--- /dev/null
+++ b/Engine/Plugins/Experimental/Animation/BlendStack/Source/Editor/Public/AnimationBlendStackGraph.h
@@ -0,0 +1,25 @@
+// Copyright Epic Games, Inc. All Rights Reserved.
+
+#pragma once
+
+#include "EdGraph/EdGraph.h"
+
+#include "AnimationBlendStackGraph.generated.h"
+
+class UObject;
+
+/** Animation graph used for blend stacks.
+*	The result node is the root of the blend graph.
+*	The input node is dynamically linked to the blend stack's sample pose.
+*/
+UCLASS(MinimalAPI)
+class UAnimationBlendStackGraph : public UEdGraph
+{
+	GENERATED_BODY()
+
+public:
+	// Result node within the state's animation graph
+	UPROPERTY()
+	TObjectPtr<class UAnimGraphNode_Root> ResultNode = nullptr;
+};
+
diff --git a/Engine/Plugins/Experimental/Animation/BlendStack/Source/Editor/Public/AnimationBlendStackGraphSchema.h b/Engine/Plugins/Experimental/Animation/BlendStack/Source/Editor/Public/AnimationBlendStackGraphSchema.h
new file mode 100644
--- /dev/null
+++ b/Engine/Plugins/Experimental/Animation/BlendStack/Source/Editor/Public/AnimationBlendStackGraphSchema.h
@@ -0,0 +1,16 @@
+// Copyright Epic Games, Inc. All Rights Reserved.
+
+#pragma once
+
+#include "AnimationGraphSchema.h"
+#include "AnimationBlendStackGraphSchema.generated.h"
+
+UCLASS(MinimalAPI)
+class UAnimationBlendStackGraphSchema : public UAnimationGraphSchema
+{
+	GENERATED_BODY()
+
+	//~ Begin UEdGraphSchema Interface.
+	virtual void CreateDefaultNodesForGraph(UEdGraph& Graph) const override;
+	//~ End UEdGraphSchema Interface.
+};
diff --git a/Engine/Plugins/Experimental/Animation/BlendStack/Source/Editor/Public/BlendStackEditor.h b/Engine/Plugins/Experimental/Animation/BlendStack/Source/Editor/Public/BlendStackEditor.h
new file mode 100644
--- /dev/null
+++ b/Engine/Plugins/Experimental/Animation/BlendStack/Source/Editor/Public/BlendStackEditor.h
@@ -0,0 +1,12 @@
+// Copyright Epic Games, Inc. All Rights Reserved.
+
+#pragma once
+
+#include "CoreMinimal.h"
+#include "Modules/ModuleInterface.h"
+
+DECLARE_LOG_CATEGORY_EXTERN(LogBlendStackEditor, Log, All);
+
+class IBlendStackEditorModuleInterface : public IModuleInterface
+{
+};
\ No newline at end of file
diff --git a/Engine/Plugins/Experimental/Animation/BlendStack/Source/Runtime/BlendStack.Build.cs b/Engine/Plugins/Experimental/Animation/BlendStack/Source/Runtime/BlendStack.Build.cs
new file mode 100644
--- /dev/null
+++ b/Engine/Plugins/Experimental/Animation/BlendStack/Source/Runtime/BlendStack.Build.cs
@@ -0,0 +1,20 @@
+// Copyright Epic Games, Inc. All Rights Reserved.
+
+using UnrealBuildTool;
+
+public class BlendStack : ModuleRules
+{
+	public BlendStack(ReadOnlyTargetRules Target) : base(Target)
+	{
+		PublicDependencyModuleNames.AddRange(
+			new string[]
+			{
+				"Core",
+				"CoreUObject",
+				"Engine",
+				"AnimationCore",
+				"AnimGraphRuntime",
+			}
+		);
+	}
+}
diff --git a/Engine/Plugins/Experimental/Animation/BlendStack/Source/Runtime/Private/AnimNode_BlendStack.cpp b/Engine/Plugins/Experimental/Animation/BlendStack/Source/Runtime/Private/AnimNode_BlendStack.cpp
new file mode 100644
--- /dev/null
+++ b/Engine/Plugins/Experimental/Animation/BlendStack/Source/Runtime/Private/AnimNode_BlendStack.cpp
@@ -0,0 +1,797 @@
+// Copyright Epic Games, Inc. All Rights Reserved.
+
+#include "BlendStack/AnimNode_BlendStack.h"
+#include "Algo/MaxElement.h"
+#include "Animation/AnimInstanceProxy.h"
+#include "Animation/AnimComposite.h"
+#include "Animation/AnimSequence.h"
+#include "Animation/BlendSpace.h"
+#include "Animation/AnimMontage.h"
+#include "BlendStack/AnimNode_BlendStackInput.h"
+#include "Animation/AnimNode_Inertialization.h"
+#include "BlendStack/BlendStackDefines.h"
+
+#include UE_INLINE_GENERATED_CPP_BY_NAME(AnimNode_BlendStack)
+
+#if ENABLE_ANIM_DEBUG
+TAutoConsoleVariable<int32> CVarAnimBlendStackEnable(TEXT("a.AnimNode.BlendStack.Enable"), 1, TEXT("Enable / Disable Blend Stack"));
+TAutoConsoleVariable<int32> CVarAnimBlendStackPruningEnable(TEXT("a.AnimNode.BlendStack.Pruning.Enable"), 1, TEXT("Enable / Disable Blend Stack Pruning"));
+#endif
+
+#define LOCTEXT_NAMESPACE "AnimNode_BlendStack"
+
+static constexpr FBoneIndexType RootBoneIndexType = 0;
+
+/////////////////////////////////////////////////////
+// FBlendStackAnimPlayer
+void FBlendStackAnimPlayer::Initialize(const FAnimationInitializeContext& Context, UAnimationAsset* AnimationAsset, float AccumulatedTime, bool bLoop, bool bMirrored, UMirrorDataTable* MirrorDataTable, float BlendTime, float RootBoneBlendTime, const UBlendProfile* BlendProfile, EAlphaBlendOption InBlendOption, FVector BlendParameters, float PlayRate, int32 InPoseLinkIdx)
+{
+	if (bMirrored && !MirrorDataTable)
+	{
+		UE_LOG(LogBlendStack, Error, TEXT("FBlendStackAnimPlayer failed to Initialize for %s. Mirroring will not work becasue MirrorDataTable is missing"), *GetNameSafe(AnimationAsset));
+	}
+	
+	check(Context.AnimInstanceProxy);
+	USkeleton* Skeleton = Context.AnimInstanceProxy->GetSkeleton();
+	check(Skeleton);
+
+	const FReferenceSkeleton& RefSkeleton = Context.AnimInstanceProxy->GetSkeleton()->GetReferenceSkeleton();
+	const bool bApplyDifferentRootBoneBlendTime = RootBoneBlendTime >= 0.f && !FMath::IsNearlyEqual(RootBoneBlendTime, BlendTime);
+	const int32 NumSkeletonBones = RefSkeleton.GetNum();
+	if (NumSkeletonBones <= 0)
+	{
+		UE_LOG(LogBlendStack, Error, TEXT("FBlendStackAnimPlayer failed to Initialize for %s. Skeleton has no bones?!"), *GetNameSafe(AnimationAsset));
+	}
+	else if (BlendTime > UE_KINDA_SMALL_NUMBER)
+	{
+		// handling BlendTime > 0 and RootBoneBlendTime >= 0
+		if (BlendProfile != nullptr)
+		{
+			check(BlendProfile->OwningSkeleton && NumSkeletonBones == BlendProfile->OwningSkeleton->GetReferenceSkeleton().GetNum());
+
+			TotalBlendInTimePerBone.Init(BlendTime, NumSkeletonBones);
+
+			BlendProfile->FillSkeletonBoneDurationsArray(TotalBlendInTimePerBone, BlendTime);
+
+			if (bApplyDifferentRootBoneBlendTime)
+			{
+				TotalBlendInTimePerBone[RootBoneIndexType] *= RootBoneBlendTime / BlendTime;
+			}
+
+			BlendTime = *Algo::MaxElement(TotalBlendInTimePerBone);
+		}
+		else if (bApplyDifferentRootBoneBlendTime)
+		{
+			TotalBlendInTimePerBone.Init(BlendTime, NumSkeletonBones);
+			TotalBlendInTimePerBone[RootBoneIndexType] *= RootBoneBlendTime / BlendTime;
+			BlendTime = FMath::Max(BlendTime, RootBoneBlendTime);
+		}
+	}
+	else if (bApplyDifferentRootBoneBlendTime)
+	{
+		// handling BlendTime ~= 0 and RootBoneBlendTime >= 0
+		TotalBlendInTimePerBone.Init(BlendTime, NumSkeletonBones);
+		TotalBlendInTimePerBone[RootBoneIndexType] = RootBoneBlendTime;
+		BlendTime = FMath::Max(BlendTime, RootBoneBlendTime);
+	}
+
+	BlendOption = InBlendOption;
+
+	TotalBlendInTime = BlendTime;
+	CurrentBlendInTime = 0.f;
+
+	MirrorNode.SetMirrorDataTable(MirrorDataTable);
+	MirrorNode.SetMirror(bMirrored);
+	
+	bool bUnsupportedAnimAsset = false;
+	if (Cast<UAnimMontage>(AnimationAsset))
+	{
+		bUnsupportedAnimAsset = true;
+	}
+	else if (UAnimSequenceBase* SequenceBase = Cast<UAnimSequenceBase>(AnimationAsset))
+	{
+		SequencePlayerNode.SetAccumulatedTime(AccumulatedTime);
+		SequencePlayerNode.SetSequence(SequenceBase);
+		SequencePlayerNode.SetLoopAnimation(bLoop);
+		SequencePlayerNode.SetPlayRate(PlayRate);
+	}
+	else if (UBlendSpace* BlendSpace = Cast<UBlendSpace>(AnimationAsset))
+	{
+		// making sure AccumulatedTime is in normalized space
+		check(AccumulatedTime >= 0.f && AccumulatedTime <= 1.f);
+
+		BlendSpacePlayerNode.SetResetPlayTimeWhenBlendSpaceChanges(false /*!bReset*/);
+		BlendSpacePlayerNode.SetAccumulatedTime(AccumulatedTime);
+		BlendSpacePlayerNode.SetBlendSpace(BlendSpace);
+		BlendSpacePlayerNode.SetLoop(bLoop);
+		BlendSpacePlayerNode.SetPlayRate(PlayRate);
+		BlendSpacePlayerNode.SetPosition(BlendParameters);
+	}
+	else if (AnimationAsset)
+	{
+		bUnsupportedAnimAsset = true;
+	}
+
+	if (bUnsupportedAnimAsset)
+	{
+		UE_LOG(LogBlendStack, Error, TEXT("FBlendStackAnimPlayer unsupported AnimationAsset %s"), *GetNameSafe(AnimationAsset));
+	}
+
+	UpdateSourceLinkNode();
+	PoseLinkIndex = InPoseLinkIdx;
+}
+
+void FBlendStackAnimPlayer::UpdatePlayRate(float PlayRate)
+{
+	if (SequencePlayerNode.GetSequence())
+	{
+		SequencePlayerNode.SetPlayRate(PlayRate);
+	}
+	else if (BlendSpacePlayerNode.GetBlendSpace())
+	{
+		BlendSpacePlayerNode.SetPlayRate(PlayRate);
+	}
+}
+
+void FBlendStackAnimPlayer::StorePoseContext(const FPoseContext& PoseContext)
+{
+	SequencePlayerNode.SetSequence(nullptr);
+	BlendSpacePlayerNode.SetBlendSpace(nullptr);
+	MirrorNode.SetSourceLinkNode(nullptr);
+
+	if (PoseContext.Pose.IsValid())
+	{
+		StoredPose.CopyBonesFrom(PoseContext.Pose);
+		StoredPose.CopyAndAssignBoneContainer(StoredBoneContainer);
+	}
+
+	StoredCurve.CopyFrom(PoseContext.Curve);
+	StoredAttributes.CopyFrom(PoseContext.CustomAttributes);
+}
+
+void FBlendStackAnimPlayer::RestorePoseContext(FPoseContext& PoseContext) const
+{
+	check(!SequencePlayerNode.GetSequence() && !BlendSpacePlayerNode.GetBlendSpace());
+
+	if (StoredPose.IsValid())
+	{
+		// Serial number mismatch means a potential bone LOD mismatch, even if we have the same number of bones.
+		// Remap the pose manually in those cases.
+		if (PoseContext.Pose.GetBoneContainer().GetSerialNumber() == StoredBoneContainer.GetSerialNumber())
+		{
+			PoseContext.Pose.CopyBonesFrom(StoredPose);
+		}
+		else
+		{
+			const FBoneContainer CurrentBoneContainer = PoseContext.Pose.GetBoneContainer();
+			for (FCompactPoseBoneIndex CompactPoseIndex : PoseContext.Pose.ForEachBoneIndex())
+			{
+				// Map the current compact pose index to skeleton index, and map this back to the stored compact pose index.
+				const FSkeletonPoseBoneIndex SkeletonPoseIndex =  CurrentBoneContainer.GetSkeletonPoseIndexFromCompactPoseIndex(CompactPoseIndex);
+				const FCompactPoseBoneIndex StoredCompactPoseIndex =  StoredBoneContainer.GetCompactPoseIndexFromSkeletonPoseIndex(SkeletonPoseIndex);
+				if (StoredCompactPoseIndex == INDEX_NONE)
+				{
+					// If our stored pose doesn't have the bone, reset to ref pose.
+					PoseContext.Pose[CompactPoseIndex] = CurrentBoneContainer.GetRefPoseTransform(CompactPoseIndex);
+				}
+				else
+				{
+					PoseContext.Pose[CompactPoseIndex] = StoredPose[StoredCompactPoseIndex];
+				}
+			}
+		}
+	}
+	else
+	{
+		PoseContext.Pose.ResetToRefPose();
+	}
+	
+	PoseContext.Curve.CopyFrom(StoredCurve);
+	PoseContext.CustomAttributes.CopyFrom(StoredAttributes);
+}
+
+
+// @todo: maybe implement copy/move constructors and assignment operator do so (or use a list instead of an array)
+// since we're making copies and moving this object in memory, we're using this method to set the MirrorNode SourceLinkNode when necessary
+void FBlendStackAnimPlayer::UpdateSourceLinkNode()
+{
+	if (SequencePlayerNode.GetSequence())
+	{
+		MirrorNode.SetSourceLinkNode(&SequencePlayerNode);
+	}
+	else if (BlendSpacePlayerNode.GetBlendSpace())
+	{
+		MirrorNode.SetSourceLinkNode(&BlendSpacePlayerNode);
+	}
+	else
+	{
+		MirrorNode.SetSourceLinkNode(nullptr);
+	}
+}
+
+void FBlendStackAnimPlayer::Evaluate_AnyThread(FPoseContext& Output)
+{
+	if (SequencePlayerNode.GetSequence() || BlendSpacePlayerNode.GetBlendSpace())
+	{
+		UpdateSourceLinkNode();
+		MirrorNode.Evaluate_AnyThread(Output);
+	}
+	else
+	{
+		RestorePoseContext(Output);
+	}
+}
+
+void FBlendStackAnimPlayer::Update_AnyThread(const FAnimationUpdateContext& Context)
+{
+	UpdateSourceLinkNode();
+	MirrorNode.Update_AnyThread(Context);
+}
+
+float FBlendStackAnimPlayer::GetAccumulatedTime() const
+{
+	if (SequencePlayerNode.GetSequence())
+	{
+		return SequencePlayerNode.GetAccumulatedTime();
+	}
+	
+	if (BlendSpacePlayerNode.GetBlendSpace())
+	{
+		// making sure BlendSpacePlayerNode.GetAccumulatedTime() is in normalized space
+		check(BlendSpacePlayerNode.GetAccumulatedTime() >= 0.f && BlendSpacePlayerNode.GetAccumulatedTime() <= 1.f);
+		return BlendSpacePlayerNode.GetAccumulatedTime();
+	}
+
+	return 0.f;
+}
+
+float FBlendStackAnimPlayer::GetPlayRate() const
+{
+	if (SequencePlayerNode.GetSequence())
+	{
+		return SequencePlayerNode.GetPlayRate();
+	}
+
+	if (BlendSpacePlayerNode.GetBlendSpace())
+	{
+		return BlendSpacePlayerNode.GetPlayRate();
+	}
+
+	return 0.f;
+}
+
+FVector FBlendStackAnimPlayer::GetBlendParameters() const
+{
+	if (BlendSpacePlayerNode.GetBlendSpace())
+	{
+		return BlendSpacePlayerNode.GetPosition();
+	}
+
+	return FVector::ZeroVector;
+}
+
+FString FBlendStackAnimPlayer::GetAnimationName() const
+{
+	if (SequencePlayerNode.GetSequence())
+	{
+		check(SequencePlayerNode.GetSequence());
+		return SequencePlayerNode.GetSequence()->GetName();
+	}
+
+	if (BlendSpacePlayerNode.GetBlendSpace())
+	{
+		check(BlendSpacePlayerNode.GetBlendSpace());
+		return BlendSpacePlayerNode.GetBlendSpace()->GetName();
+	}
+
+	return FString("StoredPose");
+}
+
+const UAnimationAsset* FBlendStackAnimPlayer::GetAnimationAsset() const
+{
+	if (SequencePlayerNode.GetSequence())
+	{
+		return SequencePlayerNode.GetSequence();
+	}
+
+	if (BlendSpacePlayerNode.GetBlendSpace())
+	{
+		return BlendSpacePlayerNode.GetBlendSpace();
+	}
+
+	return nullptr;
+}
+
+float FBlendStackAnimPlayer::GetBlendInPercentage() const
+{
+	if (FMath::IsNearlyZero(TotalBlendInTime))
+	{
+		return 1.f;
+	}
+
+	return FMath::Clamp(CurrentBlendInTime / TotalBlendInTime, 0.f, 1.f);
+}
+
+bool FBlendStackAnimPlayer::GetBlendInWeights(TArray<float>& Weights) const
+{
+	const int32 NumBones = TotalBlendInTimePerBone.Num();
+	if (NumBones > 0)
+	{
+		Weights.SetNumUninitialized(NumBones);
+		for (int32 BoneIdx = 0; BoneIdx < NumBones; ++BoneIdx)
+		{
+			const float TotalBlendInTimeBoneIdx = TotalBlendInTimePerBone[BoneIdx];
+			if (FMath::IsNearlyZero(TotalBlendInTimeBoneIdx))
+			{
+				Weights[BoneIdx] = 1.f;
+			}
+			else
+			{
+				const float UnclampedLinearWeight = CurrentBlendInTime / TotalBlendInTimeBoneIdx;
+				Weights[BoneIdx] = FAlphaBlend::AlphaToBlendOption(UnclampedLinearWeight, BlendOption);
+			}
+		}
+		return true;
+	}
+	return false;
+}
+
+/////////////////////////////////////////////////////
+// FAnimNode_BlendStack_Standalone
+void FAnimNode_BlendStack_Standalone::Evaluate_AnyThread(FPoseContext& Output)
+{
+	DECLARE_SCOPE_HIERARCHICAL_COUNTER_ANIMNODE(Evaluate_AnyThread);
+	QUICK_SCOPE_CYCLE_COUNTER(STAT_BlendStack_Evaluate_AnyThread);
+
+	Super::Evaluate_AnyThread(Output);
+
+	const int32 BlendStackSize = AnimPlayers.Num();
+	if (BlendStackSize <= 0)
+	{
+		Output.ResetToRefPose();
+	}
+	else if (BlendStackSize == 1)
+	{
+		EvaluateSample(Output, 0);
+	}
+	else if (MaxActiveBlends <= 0 
+#if ENABLE_ANIM_DEBUG
+		|| !CVarAnimBlendStackEnable.GetValueOnAnyThread()
+#endif // ENABLE_ANIM_DEBUG
+		)
+	{
+		// Disable blend stack if requested (for testing / debugging) by removing all the AnimPlayers except the first
+		while (AnimPlayers.Num() > 1)
+		{
+			AnimPlayers.PopLast();
+		}
+		EvaluateSample(Output, 0);
+	}
+	else
+	{
+		// evaluating the last AnimPlayer into Output...
+		EvaluateSample(Output, BlendStackSize - 1);
+
+		FPoseContext EvaluationPoseContext(Output);
+		FPoseContext BlendedPoseContext(Output); // @todo: this should not be necessary (but FBaseBlendedCurve::InitFrom complains about "ensure(&InCurveToInitFrom != this)"): optimize it away!
+		FAnimationPoseData BlendedAnimationPoseData(BlendedPoseContext);
+
+		const USkeleton* SkeletonAsset = Output.AnimInstanceProxy->GetRequiredBones().GetSkeletonAsset();
+		check(SkeletonAsset);
+
+		const FReferenceSkeleton& RefSkeleton = SkeletonAsset->GetReferenceSkeleton();
+		const int32 NumSkeletonBones = RefSkeleton.GetNum();
+		TArray<float> Weights;
+
+		auto EvaluateAndBlendPlayerByIndex = [this, &EvaluationPoseContext, &BlendedAnimationPoseData, &Output, &Weights, &BlendedPoseContext](int32 PlayerIndex)
+		{
+			// Evaluate into EvaluationPoseContext and then blend it with the Output (initialized with the last AnimPlayer evaluation)
+			EvaluateSample(EvaluationPoseContext, PlayerIndex);
+			if (AnimPlayers[PlayerIndex].GetBlendInWeights(Weights))
+			{
+				// @todo: have BlendTwoPosesTogetherPerBone using a TArrayView for the Weights to avoid allocations
+				FAnimationRuntime::BlendTwoPosesTogetherPerBone(FAnimationPoseData(Output), FAnimationPoseData(EvaluationPoseContext), Weights, BlendedAnimationPoseData);
+			}
+			else
+			{
+				const float Weight = 1.f - FAlphaBlend::AlphaToBlendOption(AnimPlayers[PlayerIndex].GetBlendInPercentage(), AnimPlayers[PlayerIndex].GetBlendOption());
+				FAnimationRuntime::BlendTwoPosesTogether(FAnimationPoseData(Output), FAnimationPoseData(EvaluationPoseContext), Weight, BlendedAnimationPoseData);
+			}
+			Output = BlendedPoseContext; // @todo: this should not be necessary either: optimize it away!
+		};
+
+#if ENABLE_ANIM_DEBUG
+			const bool bEnablePruning = CVarAnimBlendStackPruningEnable.GetValueOnAnyThread() > 0;
+#else
+			const bool bEnablePruning = true;
+#endif // ENABLE_ANIM_DEBUG
+
+		// Evaluate our players from the second last to the first.
+		int32 PlayerIndex = BlendStackSize - 2;
+		// Start evaluating with our least significant players.
+		for (; PlayerIndex >= MaxActiveBlends; --PlayerIndex)
+		{
+			EvaluateAndBlendPlayerByIndex(PlayerIndex);
+
+			if (bEnablePruning)
+			{
+				// too many AnimPlayers! we don't have enough available blends to hold them all, so we accumulate the blended poses into Output / BlendedPoseContext.
+				AnimPlayers.PopLast();
+			}
+		}
+
+		// Even if we're not pruning, we must use the stored pose if we have a limited number of graphs to execute.
+		const bool bNeedsStoredPose = (bEnablePruning || !SampleGraphPoseLinks.IsEmpty()) && (PlayerIndex == (MaxActiveBlends - 1));
+		if (bNeedsStoredPose)
+		{
+			check(AnimPlayers.Num() == MaxActiveBlends + 1);
+
+			// We store Output / BlendedPoseContext into the last AnimPlayer, that will hold a static pose, no longer an animation playing.
+			AnimPlayers.Last().StorePoseContext(Output);
+
+			if (!SampleGraphPoseLinks.IsEmpty())
+			{
+				const int32 PoseLinkIdx = AnimPlayers[MaxActiveBlends].GetPoseLinkIndex();
+				FBlendStack_SampleGraphPoseLink& PoseLink = SampleGraphPoseLinks[PoseLinkIdx];
+				// No players should have evaluated a graph before this point.
+				// Evaluate the graph on the blended result.
+				PoseLink.EvaluatePlayer(Output, AnimPlayers[MaxActiveBlends]);
+			}
+		}
+
+		// Continue with our most significant players.
+		for (; PlayerIndex >= 0; --PlayerIndex)
+		{
+			EvaluateAndBlendPlayerByIndex(PlayerIndex);
+		}
+
+		const int32 ActiveBlends = AnimPlayers.Num() - 1;
+		if (ActiveBlends > MaxActiveBlends)
+		{
+			UE_LOG(LogBlendStack, Display, TEXT("FAnimNode_BlendStack_Standalone NumBlends/MaxNumBlends %d / %d"), ActiveBlends, MaxActiveBlends);
+		}
+	}
+}
+
+void FAnimNode_BlendStack_Standalone::Initialize_AnyThread(const FAnimationInitializeContext& Context)
+{
+	Super::Initialize_AnyThread(Context);
+
+	Reset();
+
+	if (SampleGraphPoseLinks.IsEmpty() == false)
+	{
+		IAnimClassInterface* AnimBlueprintClass = Context.GetAnimClass();
+		check(AnimBlueprintClass);
+
+		// Patch our pose links
+		for (FBlendStack_SampleGraphPoseLink& GraphPoseLink : SampleGraphPoseLinks)
+		{
+			if (GraphPoseLink.RootNodeIndex != INDEX_NONE)
+			{
+				GraphPoseLink.Root.LinkID = AnimBlueprintClass->GetAnimNodeProperties().Num() - 1 - GraphPoseLink.RootNodeIndex;
+			}
+
+			GraphPoseLink.CacheBoneCounter.Reset();
+		}
+	}
+}
+
+void FAnimNode_BlendStack_Standalone::CacheBones_AnyThread(const FAnimationCacheBonesContext& Context) 
+{
+	const int32 BlendStackSize = AnimPlayers.Num();
+	for (int32 AnimPlayerIndex = 0; AnimPlayerIndex < BlendStackSize; ++AnimPlayerIndex)
+	{
+		// Cache bones for all active anim players.
+		// There's no need to check for weight since all unneeded anim players 
+		// would have been pruned during the last evaluation.
+		CacheBonesForSample(Context, AnimPlayerIndex);
+	}
+}
+
+void FAnimNode_BlendStack_Standalone::UpdateAssetPlayer(const FAnimationUpdateContext& Context)
+{
+	QUICK_SCOPE_CYCLE_COUNTER(STAT_BlendStack_UpdateAssetPlayer);
+
+	Super::UpdateAssetPlayer(Context);
+
+	// AnimPlayers[0] is the most newly inserted AnimPlayer, AnimPlayers[AnimPlayers.Num()-1] is the oldest, so to calculate the weights
+	// we ask AnimPlayers[0] its BlendInPercentage and then distribute the left over (CurrentWeightMultiplier) to the rest of the AnimPlayers
+	// AnimPlayers[AnimPlayerIndex].GetBlendWeight() will now store the weighted contribution of AnimPlayers[AnimPlayerIndex] to be able to calculate root motion from animation
+	float CurrentWeightMultiplier = 1.f;
+	const int32 BlendStackSize = AnimPlayers.Num();
+	int32 AnimPlayerIndex = 0;
+	for (; AnimPlayerIndex < BlendStackSize; ++AnimPlayerIndex)
+	{
+		FBlendStackAnimPlayer& AnimPlayer = AnimPlayers[AnimPlayerIndex];
+		const bool bIsLastAnimPlayers = AnimPlayerIndex == BlendStackSize - 1;
+		const float BlendInPercentage = bIsLastAnimPlayers ? 1.f : AnimPlayer.GetBlendInPercentage();
+		const float AnimPlayerBlendWeight = CurrentWeightMultiplier * BlendInPercentage;
+
+		// don't break for AnimPlayerIndex == 0 since FAnimNode_BlendStack_Standalone::BlendTo initialize the AnimPlayer with a weight of zero
+		if (AnimPlayerIndex > 0 && AnimPlayerBlendWeight < UE_KINDA_SMALL_NUMBER)
+		{
+			break;
+		}
+
+		FAnimationUpdateContext AnimPlayerContext = Context.FractionalWeightAndRootMotion(AnimPlayerBlendWeight, AnimPlayerBlendWeight);
+		UpdateSample((AnimPlayerIndex == 0) ? AnimPlayerContext : AnimPlayerContext.AsInactive(), AnimPlayerIndex);
+		CurrentWeightMultiplier *= (1.f - BlendInPercentage);
+	}
+
+	// AnimPlayers[AnimPlayerIndex] is the first FBlendStackAnimPlayer with a weight contribution of zero, so we can discard it and all the successive AnimPlayers as well
+	const int32 WantedAnimPlayersNum = FMath::Max(1, AnimPlayerIndex); // we save at least one FBlendStackAnimPlayer
+	while (AnimPlayers.Num() > WantedAnimPlayersNum)
+	{
+		AnimPlayers.PopLast();
+	}
+}
+
+bool FAnimNode_BlendStack_Standalone::IsSampleGraphAvailableForPlayer(const int32 PlayerIndex)
+{
+	// If we have any sample graphs, our player has been assigned a pose link index.
+	// If we are within X most relelvant players, then the graph is available.
+	return !SampleGraphPoseLinks.IsEmpty() && (PlayerIndex < MaxActiveBlends);
+}
+
+void FAnimNode_BlendStack_Standalone::EvaluateSample(FPoseContext& Output, const int32 PlayerIndex)
+{
+	FBlendStackAnimPlayer& SamplePlayer = AnimPlayers[PlayerIndex];
+	// If we have any sample graphs, our player has been assigned a pose link index.
+	// If we are within X most relelvant players, then the graph is available.
+	// If PlayerIndex == MaxActiveBlends, don't evaluate that graph. It's reserved for the stored pose.
+	// MaxActiveBlends == 0, means we're using inertialization. Run the the graph.
+	const bool bIsSampleGraphAvailable = !SampleGraphPoseLinks.IsEmpty() && 
+										((PlayerIndex < MaxActiveBlends) || (MaxActiveBlends == 0));
+	if (!bIsSampleGraphAvailable)
+	{
+		// If we have no sample graph, evaluate the player directly.
+		SamplePlayer.Evaluate_AnyThread(Output);
+		return;
+	}
+
+	FBlendStack_SampleGraphPoseLink& PoseLink = SampleGraphPoseLinks[SamplePlayer.GetPoseLinkIndex()];
+	PoseLink.EvaluatePlayer(Output, SamplePlayer);
+}
+
+void FBlendStack_SampleGraphPoseLink::EvaluatePlayer(FPoseContext& Output, FBlendStackAnimPlayer& SamplePlayer)
+{
+	SetInputPosePlayer(SamplePlayer);
+
+	// Make sure CacheBones has been called before evaluating.
+	ConditionalCacheBones(Output);
+	// The anim player may or may not have its Evaluate_AnyThread called through the graph update. 
+	Root.Evaluate(Output);
+}
+
+void FBlendStack_SampleGraphPoseLink::ConditionalCacheBones(const FAnimationBaseContext& Context)
+{
+	// Only call CacheBones when needed.
+	if (!CacheBoneCounter.IsSynchronized_Counter(Context.AnimInstanceProxy->GetCachedBonesCounter()))
+	{
+		// Keep track of samples that have had CacheBones called on.
+		CacheBoneCounter.SynchronizeWith(Context.AnimInstanceProxy->GetCachedBonesCounter());
+
+		FAnimationCacheBonesContext CacheBoneContext(Context.AnimInstanceProxy);
+		Root.CacheBones(CacheBoneContext);
+	}
+}
+
+void FAnimNode_BlendStack_Standalone::UpdateSample(const FAnimationUpdateContext& Context, const int32 PlayerIndex)
+{
+	FBlendStackAnimPlayer& SamplePlayer = AnimPlayers[PlayerIndex];
+
+	// If we have any sample graphs, our player has been assigned a pose link index.
+	// If we are within X most relelvant players, then the graph is available.
+	// @todo: If PlayerIndex == MaxActiveBlends, this will likely become a stored pose. What do we update that graph with? 
+	// For now, just use the same player.
+	const bool bHasSampleGraph = !SampleGraphPoseLinks.IsEmpty() && (PlayerIndex <= MaxActiveBlends);
+	if (bHasSampleGraph)
+	{
+		FBlendStack_SampleGraphPoseLink& PoseLink = SampleGraphPoseLinks[SamplePlayer.GetPoseLinkIndex()];
+		PoseLink.SetInputPosePlayer(SamplePlayer);
+		// The anim player may or may not have its Update_AnyThread called through the graph update. 
+		PoseLink.Root.Update(Context);
+	}
+	else
+	{
+		// If we have no sample graph, update the player directly.
+		SamplePlayer.Update_AnyThread(Context);
+	}
+
+	// Advance the blend-in time regardless of whether or not the player was updated.
+	SamplePlayer.AdvanceBlendInTime(Context.GetDeltaTime());
+}
+
+void FAnimNode_BlendStack_Standalone::CacheBonesForSample(const FAnimationCacheBonesContext& Context, const int32 PlayerIndex)
+{
+	FBlendStackAnimPlayer& SamplePlayer = AnimPlayers[PlayerIndex];
+	const bool bHasSampleGraph = !SampleGraphPoseLinks.IsEmpty() && (PlayerIndex <= MaxActiveBlends);
+	if (bHasSampleGraph)
+	{
+		FBlendStack_SampleGraphPoseLink& PoseLink = SampleGraphPoseLinks[SamplePlayer.GetPoseLinkIndex()];
+		PoseLink.ConditionalCacheBones(Context);
+	}
+}
+
+void FAnimNode_BlendStack_Standalone::InitializeSample(const FAnimationInitializeContext& Context, FBlendStackAnimPlayer& SamplePlayer)
+{
+	if (SamplePlayer.GetPoseLinkIndex() != INDEX_NONE)
+	{
+		FBlendStack_SampleGraphPoseLink& PoseLink = SampleGraphPoseLinks[SamplePlayer.GetPoseLinkIndex()];
+		PoseLink.SetInputPosePlayer(SamplePlayer);
+		PoseLink.Root.Initialize(Context);
+		PoseLink.ConditionalCacheBones(Context);
+	}
+}
+
+float FAnimNode_BlendStack_Standalone::GetAccumulatedTime() const
+{
+	return AnimPlayers.IsEmpty() ? 0.f : AnimPlayers.First().GetAccumulatedTime();
+}
+
+static void RequestInertialBlend(const FAnimationUpdateContext& Context, float BlendTime, const UBlendProfile* BlendProfile, EAlphaBlendOption BlendOption)
+{
+	if (BlendTime > 0.0f)
+	{
+		UE::Anim::IInertializationRequester* InertializationRequester = Context.GetMessage<UE::Anim::IInertializationRequester>();
+		if (InertializationRequester)
+		{
+			FInertializationRequest Request;
+			Request.Duration = BlendTime;
+			Request.BlendProfile = BlendProfile;
+			Request.bUseBlendMode = true;
+			Request.BlendMode = BlendOption;
+#if ANIM_TRACE_ENABLED
+			Request.Description = LOCTEXT("InertializationRequestDescription", "Blend Stack");
+			Request.NodeId = Context.GetCurrentNodeId();
+			Request.AnimInstance = Context.AnimInstanceProxy->GetAnimInstanceObject();
+#endif
+
+			InertializationRequester->RequestInertialization(Request);
+		}
+		else
+		{
+			FAnimNode_Inertialization::LogRequestError(Context, Context.GetCurrentNodeId());
+		}
+	}
+}
+
+void FAnimNode_BlendStack_Standalone::BlendTo(const FAnimationUpdateContext& Context, UAnimationAsset* AnimationAsset, float AccumulatedTime, bool bLoop, bool bMirrored, UMirrorDataTable* MirrorDataTable, float BlendTime, float RootBoneBlendTime, const UBlendProfile* BlendProfile, EAlphaBlendOption BlendOption, bool bUseInertialBlend, FVector BlendParameters, float PlayRate)
+{
+	if (bUseInertialBlend)
+	{
+		RequestInertialBlend(Context, BlendTime, BlendProfile, BlendOption);
+		BlendTime = 0.0f;
+	}
+
+	AnimPlayers.PushFirst(FBlendStackAnimPlayer());
+	FBlendStackAnimPlayer& AnimPlayer = AnimPlayers.First();
+
+	FAnimationInitializeContext InitContext(Context.AnimInstanceProxy, Context.SharedContext);
+	AnimPlayer.Initialize(InitContext, AnimationAsset, AccumulatedTime, bLoop, bMirrored, MirrorDataTable, BlendTime, RootBoneBlendTime, BlendProfile, BlendOption, BlendParameters, PlayRate, GetNextPoseLinkIndex());
+	InitializeSample(InitContext, AnimPlayer);
+}
+
+void FAnimNode_BlendStack_Standalone::Reset()
+{
+	AnimPlayers.Reset();
+}
+
+int32 FAnimNode_BlendStack_Standalone::GetNextPoseLinkIndex()
+{
+	if (SampleGraphPoseLinks.IsEmpty())
+	{
+		return INDEX_NONE;
+	}
+
+	const int32 NumPoseLinks = SampleGraphPoseLinks.Num();
+	CurrentSamplePoseLink = ++CurrentSamplePoseLink;
+	if (CurrentSamplePoseLink == NumPoseLinks) { CurrentSamplePoseLink = 0; }
+
+	return CurrentSamplePoseLink;
+}
+
+void FAnimNode_BlendStack_Standalone::UpdatePlayRate(float PlayRate)
+{
+	if (!AnimPlayers.IsEmpty())
+	{
+		AnimPlayers.First().UpdatePlayRate(PlayRate);
+	}
+}
+
+void FAnimNode_BlendStack_Standalone::GatherDebugData(FNodeDebugData& DebugData)
+{
+#if ENABLE_ANIM_DEBUG
+	DebugData.AddDebugItem(FString::Printf(TEXT("%s"), *DebugData.GetNodeName(this)));
+	for (int32 i = 0; i < AnimPlayers.Num(); ++i)
+	{
+		const FBlendStackAnimPlayer& AnimPlayer = AnimPlayers[i];
+		DebugData.AddDebugItem(FString::Printf(TEXT("%d) t:%.2f/%.2f m:%d %s"),
+				i, AnimPlayer.GetCurrentBlendInTime(), AnimPlayer.GetTotalBlendInTime(),
+				AnimPlayer.GetMirror() ? 1 : 0, *AnimPlayer.GetAnimationName()));
+	}
+#endif // ENABLE_ANIM_DEBUG
+
+	// propagating GatherDebugData to the AnimPlayers
+	for (FBlendStackAnimPlayer& AnimPlayer : AnimPlayers)
+	{
+		AnimPlayer.GetMirrorNode().GatherDebugData(DebugData);
+	}
+}
+
+/////////////////////////////////////////////////////
+// FAnimNode_BlendStack
+
+void FAnimNode_BlendStack::UpdateAssetPlayer(const FAnimationUpdateContext& Context)
+{
+	const bool bNeedsReset =
+		bResetOnBecomingRelevant &&
+		UpdateCounter.HasEverBeenUpdated() &&
+		!UpdateCounter.WasSynchronizedCounter(Context.AnimInstanceProxy->GetUpdateCounter());
+
+	if (bNeedsReset)
+	{
+		Reset();
+	}
+
+	UpdateCounter.SynchronizeWith(Context.AnimInstanceProxy->GetUpdateCounter());
+
+	GetEvaluateGraphExposedInputs().Execute(Context);
+
+	bool bExecuteBlendTo = false;
+	if (AnimPlayers.IsEmpty())
+	{
+		bExecuteBlendTo = true;
+	}
+	else
+	{
+		const FBlendStackAnimPlayer& MainAnimPlayer = AnimPlayers.First();
+		const UAnimationAsset* PlayingAnimationAsset = MainAnimPlayer.GetAnimationAsset();
+
+		if (bForceBlendNextUpdate)
+		{
+			bForceBlendNextUpdate = false;
+			bExecuteBlendTo = true;
+		}
+		else if (AnimationAsset != PlayingAnimationAsset)
+		{
+			bExecuteBlendTo = true;
+		}
+		else if (bMirrored != MainAnimPlayer.GetMirror())
+		{
+			bExecuteBlendTo = true;
+		}
+		else if (BlendParameters != MainAnimPlayer.GetBlendParameters())
+		{
+			bExecuteBlendTo = true;
+		}
+		else if (MaxAnimationDeltaTime >= 0.f && FMath::Abs(AnimationTime - MainAnimPlayer.GetAccumulatedTime()) > MaxAnimationDeltaTime)
+		{
+			bExecuteBlendTo = true;
+		}
+	}
+
+	if (bExecuteBlendTo)
+	{
+		BlendTo(Context, AnimationAsset, AnimationTime, bLoop, bMirrored, MirrorDataTable.Get(), BlendTime, RootBoneBlendTime, BlendProfile, BlendOption, bUseInertialBlend, BlendParameters, WantedPlayRate);
+	}
+	
+	UpdatePlayRate(WantedPlayRate);
+
+	Super::UpdateAssetPlayer(Context);
+}
+
+void FAnimNode_BlendStack::ForceBlendNextUpdate()
+{
+	bForceBlendNextUpdate = true;
+}
+
+void FBlendStack_SampleGraphPoseLink::SetInputPosePlayer(FBlendStackAnimPlayer& InPlayer)
+{
+	// Because our anim players may get reallocated, or change indices due to push/pops,
+	// we must call this before every operation that might end up needing the anim player through the graph's input nodes.
+	Player = &InPlayer;
+}
+
+#undef LOCTEXT_NAMESPACE
\ No newline at end of file
diff --git a/Engine/Plugins/Experimental/Animation/BlendStack/Source/Runtime/Private/AnimNode_BlendStackInput.cpp b/Engine/Plugins/Experimental/Animation/BlendStack/Source/Runtime/Private/AnimNode_BlendStackInput.cpp
new file mode 100644
--- /dev/null
+++ b/Engine/Plugins/Experimental/Animation/BlendStack/Source/Runtime/Private/AnimNode_BlendStackInput.cpp
@@ -0,0 +1,42 @@
+// Copyright Epic Games, Inc. All Rights Reserved.
+
+#include "BlendStack/AnimNode_BlendStackInput.h"
+#include "BlendStack/AnimNode_BlendStack.h"
+#include "Animation/AnimInstanceProxy.h"
+
+#include UE_INLINE_GENERATED_CPP_BY_NAME(AnimNode_BlendStackInput)
+
+void FAnimNode_BlendStackInput::Evaluate_AnyThread(FPoseContext& Output)
+{
+	check(Player && *Player);
+	(*Player)->Evaluate_AnyThread(Output);
+}
+
+static FAnimNode_BlendStack_Standalone* GetBlendStackNodeFromIndex(const FAnimationBaseContext& Context, IAnimClassInterface* AnimBlueprintClass, int32 NodeIndex)
+{
+	const TArray<FStructProperty*>& AnimNodeProperties = AnimBlueprintClass->GetAnimNodeProperties();
+	check(AnimNodeProperties.IsValidIndex(NodeIndex));
+
+	FStructProperty* LinkedProperty = AnimNodeProperties[NodeIndex];
+	return LinkedProperty->ContainerPtrToValuePtr<FAnimNode_BlendStack_Standalone>(Context.AnimInstanceProxy->GetAnimInstanceObject());
+}
+
+void FAnimNode_BlendStackInput::Update_AnyThread(const FAnimationUpdateContext& Context)
+{
+	// If there's no player, use the allocated blend stack index to get a reference to it.
+	// This should only happen on the first ever update of this node.
+	if (Player == nullptr)
+	{
+		IAnimClassInterface* AnimBlueprintClass = Context.GetAnimClass();
+		check(AnimBlueprintClass);
+		// Allocation index is the inverse of the node index.
+		const int32 BlendStackNodeIndex = AnimBlueprintClass->GetAnimNodeProperties().Num() - 1 - BlendStackAllocationIndex;
+		check(Context.AnimInstanceProxy);
+		FAnimNode_BlendStack_Standalone* BlendStackNode = GetBlendStackNodeFromIndex(Context, AnimBlueprintClass, BlendStackNodeIndex);
+		Player = &BlendStackNode->SampleGraphPoseLinks[SampleIndex].Player;
+	}
+
+	check(Player && *Player);
+	(*Player)->Update_AnyThread(Context);
+}
+
diff --git a/Engine/Plugins/Experimental/Animation/BlendStack/Source/Runtime/Private/BlendStackAnimNodeLibrary.cpp b/Engine/Plugins/Experimental/Animation/BlendStack/Source/Runtime/Private/BlendStackAnimNodeLibrary.cpp
new file mode 100644
--- /dev/null
+++ b/Engine/Plugins/Experimental/Animation/BlendStack/Source/Runtime/Private/BlendStackAnimNodeLibrary.cpp
@@ -0,0 +1,68 @@
+// Copyright Epic Games, Inc. All Rights Reserved.
+
+#include "BlendStack/BlendStackAnimNodeLibrary.h"
+#include "BlendStack/AnimNode_BlendStack.h"
+#include "BlendStack/BlendStackDefines.h"
+
+FBlendStackAnimNodeReference UBlendStackAnimNodeLibrary::ConvertToBlendStackNode(const FAnimNodeReference& Node, EAnimNodeReferenceConversionResult& Result)
+{
+	return FAnimNodeReference::ConvertToType<FBlendStackAnimNodeReference>(Node, Result);
+}
+
+void UBlendStackAnimNodeLibrary::BlendTo(const FAnimUpdateContext& Context, 
+										const FBlendStackAnimNodeReference& BlendStackNode, 
+										UAnimationAsset* AnimationAsset,
+										float AnimationTime,
+										bool bLoop,
+										bool bMirrored,
+										float BlendTime,
+										FVector BlendParameters,
+										float WantedPlayRate)
+{
+	if (FAnimNode_BlendStack* BlendStackNodePtr = BlendStackNode.GetAnimNodePtr<FAnimNode_BlendStack>())
+	{
+		if (const FAnimationUpdateContext* AnimationUpdateContext = Context.GetContext())
+		{
+			if (AnimationAsset == nullptr)
+			{
+				UE_LOG(LogBlendStack, Warning, TEXT("UBlendStackAnimNodeLibrary::BlendTo called with null animation asset."));
+				return;
+			}
+
+			BlendStackNodePtr->BlendTo(
+				*AnimationUpdateContext, 
+				AnimationAsset, 
+				AnimationTime, 
+				bLoop,
+				bMirrored,
+				BlendStackNodePtr->MirrorDataTable,
+				BlendTime,
+				BlendStackNodePtr->RootBoneBlendTime,
+				BlendStackNodePtr->BlendProfile,
+				BlendStackNodePtr->BlendOption,
+				BlendStackNodePtr->bUseInertialBlend,
+				BlendParameters,
+				WantedPlayRate);
+		}
+		else
+		{
+			UE_LOG(LogBlendStack, Warning, TEXT("UBlendStackAnimNodeLibrary::BlendTo called with an invalid context."));
+		}
+	}
+	else
+	{
+		UE_LOG(LogBlendStack, Warning, TEXT("UBlendStackAnimNodeLibrary::BlendTo called with an invalid type."));
+	}
+}
+
+void UBlendStackAnimNodeLibrary::ForceBlendNextUpdate(const FBlendStackAnimNodeReference& BlendStackNode)
+{
+	if (FAnimNode_BlendStack* BlendStackNodePtr = BlendStackNode.GetAnimNodePtr<FAnimNode_BlendStack>())
+	{
+		BlendStackNodePtr->ForceBlendNextUpdate();
+	}
+	else
+	{
+		UE_LOG(LogBlendStack, Warning, TEXT("UBlendStackAnimNodeLibrary::ForceBlendNextUpdate called with an invalid type."));
+	}
+}
\ No newline at end of file
diff --git a/Engine/Plugins/Experimental/Animation/BlendStack/Source/Runtime/Private/BlendStackDefines.cpp b/Engine/Plugins/Experimental/Animation/BlendStack/Source/Runtime/Private/BlendStackDefines.cpp
new file mode 100644
--- /dev/null
+++ b/Engine/Plugins/Experimental/Animation/BlendStack/Source/Runtime/Private/BlendStackDefines.cpp
@@ -0,0 +1,8 @@
+// Copyright Epic Games, Inc. All Rights Reserved.
+
+#include "BlendStack/BlendStackDefines.h"
+#include "Modules/ModuleManager.h"
+
+IMPLEMENT_MODULE(FDefaultModuleImpl, BlendStack);
+
+DEFINE_LOG_CATEGORY(LogBlendStack);
diff --git a/Engine/Plugins/Experimental/Animation/BlendStack/Source/Runtime/Public/BlendStack/AnimNode_BlendStack.h b/Engine/Plugins/Experimental/Animation/BlendStack/Source/Runtime/Public/BlendStack/AnimNode_BlendStack.h
new file mode 100644
--- /dev/null
+++ b/Engine/Plugins/Experimental/Animation/BlendStack/Source/Runtime/Public/BlendStack/AnimNode_BlendStack.h
@@ -0,0 +1,210 @@
+// Copyright Epic Games, Inc. All Rights Reserved.
+
+#pragma once
+
+#include "Animation/AnimNode_AssetPlayerBase.h"
+#include "Animation/AnimNode_SequencePlayer.h"
+#include "AnimNodes/AnimNode_BlendSpacePlayer.h"
+#include "AnimNodes/AnimNode_Mirror.h"
+#include "Containers/Deque.h"
+#include "AnimNode_BlendStack.generated.h"
+
+USTRUCT()
+struct BLENDSTACK_API FBlendStackAnimPlayer
+{
+	GENERATED_BODY()
+	
+	void Initialize(const FAnimationInitializeContext& Context, UAnimationAsset* AnimationAsset, float AccumulatedTime, bool bLoop, bool bMirrored, UMirrorDataTable* MirrorDataTable, float BlendTime, float RootBoneBlendTime, const UBlendProfile* BlendProfile, EAlphaBlendOption InBlendOption, FVector BlendParameters, float PlayRate, int32 InPoseLinkIdx);
+	void UpdatePlayRate(float PlayRate);
+	void Evaluate_AnyThread(FPoseContext& Output);
+	void Update_AnyThread(const FAnimationUpdateContext& Context);
+	float GetAccumulatedTime() const;
+	float GetPlayRate() const;
+	
+	float GetBlendInPercentage() const;
+	bool GetBlendInWeights(TArray<float>& Weights) const;
+	EAlphaBlendOption GetBlendOption() const { return BlendOption; }
+	void StorePoseContext(const FPoseContext& PoseContext);
+
+	float GetTotalBlendInTime() const { return TotalBlendInTime; }
+	float GetCurrentBlendInTime() const { return CurrentBlendInTime; }
+	void AdvanceBlendInTime(const float DeltaTime) { CurrentBlendInTime += DeltaTime; }
+	bool GetMirror() const { return MirrorNode.GetMirror(); }
+	FVector GetBlendParameters() const;
+	FString GetAnimationName() const;
+	const UAnimationAsset* GetAnimationAsset() const;
+
+	FAnimNode_Mirror_Standalone& GetMirrorNode() { return MirrorNode; }
+	int32 GetPoseLinkIndex() const { return PoseLinkIndex; }
+
+public:
+	void RestorePoseContext(FPoseContext& PoseContext) const;
+	void UpdateSourceLinkNode();
+
+	// Embedded standalone player to play sequence
+	FAnimNode_SequencePlayer_Standalone SequencePlayerNode;
+
+	// Embedded standalone player to play blend spaces
+	FAnimNode_BlendSpacePlayer_Standalone BlendSpacePlayerNode;
+
+	// Embedded mirror node to handle mirroring
+	FAnimNode_Mirror_Standalone MirrorNode;
+
+	// if SequencePlayerNode.GetSequence() and BlendSpacePlayerNode.GetBlendSpace() are nullptr, 
+	// instead of using SequencePlayerNode or BlendSpacePlayerNode (wrapped in MirrorNode),
+	// the output FPoseContext will be from StoredPose, StoredCurve, StoredAttributes
+	FCompactHeapPose StoredPose;
+	FBlendedHeapCurve StoredCurve;
+	UE::Anim::FHeapAttributeContainer StoredAttributes;
+	// We need to store the bone container, in case we have a LOD swap during a blend that uses the stored pose.
+	FBoneContainer StoredBoneContainer;
+
+	EAlphaBlendOption BlendOption = EAlphaBlendOption::Linear;
+	int32 PoseLinkIndex = INDEX_NONE;
+	TCustomBoneIndexArray<float, FSkeletonPoseBoneIndex> TotalBlendInTimePerBone;
+
+	float TotalBlendInTime = 0.f;
+	float CurrentBlendInTime = 0.f;
+};
+
+USTRUCT(BlueprintInternalUseOnly)
+struct FBlendStack_SampleGraphPoseLink
+{
+	GENERATED_BODY()
+
+	UPROPERTY()
+	int32 RootNodeIndex = INDEX_NONE;
+
+	FPoseLink Root;
+	FBlendStackAnimPlayer* Player = nullptr;
+	FGraphTraversalCounter CacheBoneCounter;
+
+	void SetInputPosePlayer(FBlendStackAnimPlayer& InPlayer);
+	void EvaluatePlayer(FPoseContext& Output, FBlendStackAnimPlayer& SamplePlayer);
+	void ConditionalCacheBones(const FAnimationBaseContext& Output);
+};
+
+USTRUCT(BlueprintInternalUseOnly)
+struct BLENDSTACK_API FAnimNode_BlendStack_Standalone : public FAnimNode_AssetPlayerBase
+{
+	GENERATED_BODY()
+
+	UPROPERTY()
+	TArray<FBlendStack_SampleGraphPoseLink> SampleGraphPoseLinks;
+	int32 CurrentSamplePoseLink = -1;
+
+	TDeque<FBlendStackAnimPlayer> AnimPlayers;
+
+	// FAnimNode_Base interface
+	virtual void Evaluate_AnyThread(FPoseContext& Output) override;
+	virtual void Initialize_AnyThread(const FAnimationInitializeContext& Context) override;
+	virtual void CacheBones_AnyThread(const FAnimationCacheBonesContext& Context) override;
+	virtual void GatherDebugData(FNodeDebugData& DebugData) override;
+	// End of FAnimNode_Base interface
+
+	void BlendTo(const FAnimationUpdateContext& Context, UAnimationAsset* AnimationAsset, float AccumulatedTime = 0.f, bool bLoop = false, bool bMirrored = false, UMirrorDataTable* MirrorDataTable = nullptr, float BlendTime = 0.2f, float RootBoneBlendTime = -1.f, const UBlendProfile* BlendProfile = nullptr, EAlphaBlendOption BlendOption = EAlphaBlendOption::Linear, bool bUseInertialBlend = false, FVector BlendParameters = FVector::Zero(), float PlayRate = 1.f);
+	void UpdatePlayRate(float PlayRate);
+	void Reset();
+
+	// FAnimNode_AssetPlayerBase interface
+	virtual float GetAccumulatedTime() const override;
+	virtual void UpdateAssetPlayer(const FAnimationUpdateContext& Context) override;
+	// End of FAnimNode_AssetPlayerBase interface
+
+	int32 GetNextPoseLinkIndex();
+
+	int32 GetMaxActiveBlends() const { return MaxActiveBlends; }
+
+protected:
+	// Number of max active blending animation in the blend stack. If MaxActiveBlends is zero then blend stack is disabled
+	UPROPERTY(EditAnywhere, Category = Settings, meta = (ClampMin = "0"))
+	int32 MaxActiveBlends = 4;
+
+private:
+
+	void InitializeSample(const FAnimationInitializeContext& Context, FBlendStackAnimPlayer& SamplePlayer);
+	void EvaluateSample(FPoseContext& Output, const int32 PlayerIndex);
+	void UpdateSample(const FAnimationUpdateContext& Context, const int32 PlayerIndex);
+	void CacheBonesForSample(const FAnimationCacheBonesContext& Context, const int32 PlayerIndex);
+	bool IsSampleGraphAvailableForPlayer(const int32 PlayerIndex);
+
+	friend class UAnimGraphNode_BlendStack_Base;
+};
+
+
+USTRUCT(BlueprintInternalUseOnly)
+struct BLENDSTACK_API FAnimNode_BlendStack : public FAnimNode_BlendStack_Standalone
+{
+	GENERATED_BODY()
+
+	// requested animation to play
+	UPROPERTY(EditAnywhere, BlueprintReadWrite, Category = Settings, meta = (PinShownByDefault))
+	TObjectPtr<UAnimationAsset> AnimationAsset;
+
+	// requested animation time. If negative, the animation will play from the beginning uninterrupted
+	UPROPERTY(EditAnywhere, BlueprintReadWrite, Category = Settings, meta = (PinHiddenByDefault))
+	float AnimationTime = -1.f;
+
+	// requested AnimationAsset looping
+	UPROPERTY(EditAnywhere, BlueprintReadWrite, Category = Settings, meta = (PinHiddenByDefault))
+	bool bLoop = true;
+
+	// requested AnimationAsset mirroring
+	UPROPERTY(EditAnywhere, BlueprintReadWrite, Category = Settings, meta = (PinHiddenByDefault))
+	bool bMirrored = false;
+
+	// requested animation play rate
+	UPROPERTY(EditAnywhere, BlueprintReadWrite, Category = Settings, meta = (PinHiddenByDefault))
+	float WantedPlayRate = 1.f;
+	
+	// tunable animation transition blend time 
+	UPROPERTY(EditAnywhere, BlueprintReadWrite, Category = Settings, meta = (PinHiddenByDefault, ClampMin = "0"))
+	float BlendTime = 0.2f;
+
+	// Time in seconds to blend out to the new pose root bone. Negative values would imply RootBoneBlendTime is equal to BlendTime
+	UPROPERTY(EditAnywhere, BlueprintReadWrite, Category = Settings, meta = (PinHiddenByDefault))
+	float RootBoneBlendTime = -1.f;
+
+	// if AnimationTime and MaxAnimationDeltaTime are positive and the currently playing animation total time differs more than MaxAnimationDeltaTime from AnimationTime
+	// (animation desynchronized from the requested time) the blend stack will force a blend into the same animation
+	UPROPERTY(EditAnywhere, BlueprintReadWrite, Category = Settings, meta = (PinHiddenByDefault))
+	float MaxAnimationDeltaTime = -1.f;
+
+	UPROPERTY(EditAnywhere, BlueprintReadWrite, Category = Settings, meta = (PinHiddenByDefault, UseAsBlendProfile = true))
+	TObjectPtr<UBlendProfile> BlendProfile;
+
+	UPROPERTY(EditAnywhere, BlueprintReadWrite, Category = Settings, meta = (PinHiddenByDefault))
+	EAlphaBlendOption BlendOption = EAlphaBlendOption::Linear;
+
+	// requested blend space blend parameters (if AnimationAsset is a blend space)
+	UPROPERTY(EditAnywhere, BlueprintReadWrite, Category = Settings, meta = (PinHiddenByDefault))
+	FVector BlendParameters = FVector::Zero();
+
+	// if set and bMirrored MirrorDataTable will be used for mirroring the aniamtion
+	UPROPERTY(EditAnywhere, Category = Settings, meta = (PinHiddenByDefault))
+	TObjectPtr<UMirrorDataTable> MirrorDataTable;
+
+	// tunable animation transition blend time 
+	UPROPERTY(EditAnywhere, Category = Settings)
+	bool bUseInertialBlend = false;
+
+	// Reset the blend stack if it has become relevant to the graph after not being updated on previous frames.
+	UPROPERTY(EditAnywhere, Category = Settings)
+	bool bResetOnBecomingRelevant = true;
+
+	// Update Counter for detecting being relevant
+	FGraphTraversalCounter UpdateCounter;
+
+	// FAnimNode_Base interface
+	virtual void UpdateAssetPlayer(const FAnimationUpdateContext& Context) override;
+	// End of FAnimNode_Base interface
+
+	// Force a blend on the next update, even if the anim sequence has not changed.
+	void ForceBlendNextUpdate();
+
+private:
+
+	bool bForceBlendNextUpdate = false;
+
+};
+
diff --git a/Engine/Plugins/Experimental/Animation/BlendStack/Source/Runtime/Public/BlendStack/AnimNode_BlendStackInput.h b/Engine/Plugins/Experimental/Animation/BlendStack/Source/Runtime/Public/BlendStack/AnimNode_BlendStackInput.h
new file mode 100644
--- /dev/null
+++ b/Engine/Plugins/Experimental/Animation/BlendStack/Source/Runtime/Public/BlendStack/AnimNode_BlendStackInput.h
@@ -0,0 +1,30 @@
+// Copyright Epic Games, Inc. All Rights Reserved.
+
+#pragma once
+
+#include "Animation/AnimNodeBase.h"
+#include "AnimNode_BlendStackInput.generated.h"
+
+/** Input pose that links the blend stack's sample graph with the sample/pose chosen by the blend stack.
+*	@Todo: It might be better to reuse FAnimNode_LinkedInputPose, since we will most likely need variable input pins in the future too.
+*/
+USTRUCT(BlueprintInternalUseOnly, Experimental)
+struct BLENDSTACK_API FAnimNode_BlendStackInput : public FAnimNode_Base
+{
+	GENERATED_BODY()
+	
+public:
+
+	virtual void Update_AnyThread(const FAnimationUpdateContext& Context) override;
+	virtual void Evaluate_AnyThread(FPoseContext& Output) override;
+
+	UPROPERTY()
+	int32 SampleIndex = INDEX_NONE;
+
+	UPROPERTY()
+	int32 BlendStackAllocationIndex = INDEX_NONE;
+
+	// The player is guaranteed to be valid for the whole duration of update/eval.
+	struct FBlendStackAnimPlayer** Player;
+};
+
diff --git a/Engine/Plugins/Experimental/Animation/BlendStack/Source/Runtime/Public/BlendStack/BlendStackAnimNodeLibrary.h b/Engine/Plugins/Experimental/Animation/BlendStack/Source/Runtime/Public/BlendStack/BlendStackAnimNodeLibrary.h
new file mode 100644
--- /dev/null
+++ b/Engine/Plugins/Experimental/Animation/BlendStack/Source/Runtime/Public/BlendStack/BlendStackAnimNodeLibrary.h
@@ -0,0 +1,53 @@
+// Copyright Epic Games, Inc. All Rights Reserved.
+
+#pragma once
+
+#include "Animation/AnimExecutionContext.h"
+#include "Animation/AnimNodeReference.h"
+#include "Kismet/BlueprintFunctionLibrary.h"
+#include "BlendStackAnimNodeLibrary.generated.h"
+
+struct FAnimNode_BlendStack;
+
+USTRUCT(BlueprintType)
+struct FBlendStackAnimNodeReference : public FAnimNodeReference
+{
+	GENERATED_BODY()
+
+	typedef FAnimNode_BlendStack FInternalNodeType;
+};
+
+// Exposes operations that can be run on a Blend Stack node via Anim Node Functions such as "On Become Relevant" and "On Update".
+UCLASS(Experimental)
+class BLENDSTACK_API UBlendStackAnimNodeLibrary : public UBlueprintFunctionLibrary
+{
+	GENERATED_BODY()
+
+public:
+	/** Get a blend stack node context from an anim node context */
+	UFUNCTION(BlueprintCallable, Category = "Animation|BlendStack", meta = (BlueprintThreadSafe, ExpandEnumAsExecs = "Result"))
+	static FBlendStackAnimNodeReference ConvertToBlendStackNode(const FAnimNodeReference& Node, EAnimNodeReferenceConversionResult& Result);
+
+	/** Get a blend stack node context from an anim node context (pure) */
+	UFUNCTION(BlueprintPure, Category = "Animation|BlendStack", meta = (BlueprintThreadSafe, DisplayName = "Convert to Blend Stack Node"))
+	static void ConvertToBlendStackNodePure(const FAnimNodeReference& Node, FBlendStackAnimNodeReference& BlendStackNode, bool& Result)
+	{
+		EAnimNodeReferenceConversionResult ConversionResult;
+		BlendStackNode = ConvertToBlendStackNode(Node, ConversionResult);
+		Result = (ConversionResult == EAnimNodeReferenceConversionResult::Succeeded);
+	}
+
+	UFUNCTION(BlueprintCallable, Category = "Animation|BlendStack", meta = (BlueprintThreadSafe, DisplayName = "Force Blend On Next Update"))
+	static void ForceBlendNextUpdate(const FBlendStackAnimNodeReference& BlendStackNode);
+
+	UFUNCTION(BlueprintCallable, Category = "Animation|BlendStack", meta = (BlueprintThreadSafe, DisplayName = "Blend To"))
+	static void BlendTo(const FAnimUpdateContext& Context, 
+						const FBlendStackAnimNodeReference& BlendStackNode, 
+						UAnimationAsset* AnimationAsset = nullptr,
+						float AnimationTime = 0.f,
+						bool bLoop = false,
+						bool bMirrored = false,
+						float BlendTime = 0.2f,
+						FVector BlendParameters = FVector::ZeroVector,
+						float WantedPlayRate = 1.f);
+};
\ No newline at end of file
diff --git a/Engine/Plugins/Experimental/Animation/BlendStack/Source/Runtime/Public/BlendStack/BlendStackDefines.h b/Engine/Plugins/Experimental/Animation/BlendStack/Source/Runtime/Public/BlendStack/BlendStackDefines.h
new file mode 100644
--- /dev/null
+++ b/Engine/Plugins/Experimental/Animation/BlendStack/Source/Runtime/Public/BlendStack/BlendStackDefines.h
@@ -0,0 +1,11 @@
+// Copyright Epic Games, Inc. All Rights Reserved.
+
+#pragma once
+
+#include "Logging/LogMacros.h"
+
+DECLARE_LOG_CATEGORY_EXTERN(LogBlendStack, Log, All);
+
+
+
+
diff --git a/Engine/Plugins/Experimental/Animation/PoseSearch/PoseSearch.uplugin b/Engine/Plugins/Experimental/Animation/PoseSearch/PoseSearch.uplugin
--- a/Engine/Plugins/Experimental/Animation/PoseSearch/PoseSearch.uplugin
+++ b/Engine/Plugins/Experimental/Animation/PoseSearch/PoseSearch.uplugin
@@ -38,6 +38,10 @@
 		{
 			"Name": "StructUtils",
 			"Enabled": true
+		},
+		{
+			"Name": "BlendStack",
+			"Enabled": true
 		}
 	]
 }
\ No newline at end of file
diff --git a/Engine/Plugins/Experimental/Animation/PoseSearch/Source/Editor/PoseSearchEditor.Build.cs b/Engine/Plugins/Experimental/Animation/PoseSearch/Source/Editor/PoseSearchEditor.Build.cs
--- a/Engine/Plugins/Experimental/Animation/PoseSearch/Source/Editor/PoseSearchEditor.Build.cs
+++ b/Engine/Plugins/Experimental/Animation/PoseSearch/Source/Editor/PoseSearchEditor.Build.cs
@@ -12,6 +12,8 @@ public class PoseSearchEditor : ModuleRules
 		PublicDependencyModuleNames.AddRange(
 			new string[]
 			{
+				"BlendStackEditor",
+				"BlendStack",
 				"AnimGraph",
 				"AnimGraphRuntime",
 				"AnimationCore",
diff --git a/Engine/Plugins/Experimental/Animation/PoseSearch/Source/Editor/Private/AnimGraphNode_BlendStack.cpp b/Engine/Plugins/Experimental/Animation/PoseSearch/Source/Editor/Private/AnimGraphNode_BlendStack.cpp
deleted file mode 100644
diff --git a/Engine/Plugins/Experimental/Animation/PoseSearch/Source/Editor/Private/AnimGraphNode_BlendStackInput.cpp b/Engine/Plugins/Experimental/Animation/PoseSearch/Source/Editor/Private/AnimGraphNode_BlendStackInput.cpp
deleted file mode 100644
diff --git a/Engine/Plugins/Experimental/Animation/PoseSearch/Source/Editor/Private/AnimGraphNode_BlendStackResult.cpp b/Engine/Plugins/Experimental/Animation/PoseSearch/Source/Editor/Private/AnimGraphNode_BlendStackResult.cpp
deleted file mode 100644
diff --git a/Engine/Plugins/Experimental/Animation/PoseSearch/Source/Editor/Private/AnimationBlendStackGraphSchema.cpp b/Engine/Plugins/Experimental/Animation/PoseSearch/Source/Editor/Private/AnimationBlendStackGraphSchema.cpp
deleted file mode 100644
diff --git a/Engine/Plugins/Experimental/Animation/PoseSearch/Source/Editor/Public/AnimGraphNode_BlendStack.h b/Engine/Plugins/Experimental/Animation/PoseSearch/Source/Editor/Public/AnimGraphNode_BlendStack.h
deleted file mode 100644
diff --git a/Engine/Plugins/Experimental/Animation/PoseSearch/Source/Editor/Public/AnimGraphNode_BlendStackInput.h b/Engine/Plugins/Experimental/Animation/PoseSearch/Source/Editor/Public/AnimGraphNode_BlendStackInput.h
deleted file mode 100644
diff --git a/Engine/Plugins/Experimental/Animation/PoseSearch/Source/Editor/Public/AnimGraphNode_BlendStackResult.h b/Engine/Plugins/Experimental/Animation/PoseSearch/Source/Editor/Public/AnimGraphNode_BlendStackResult.h
deleted file mode 100644
diff --git a/Engine/Plugins/Experimental/Animation/PoseSearch/Source/Editor/Public/AnimationBlendStackGraph.h b/Engine/Plugins/Experimental/Animation/PoseSearch/Source/Editor/Public/AnimationBlendStackGraph.h
deleted file mode 100644
diff --git a/Engine/Plugins/Experimental/Animation/PoseSearch/Source/Editor/Public/AnimationBlendStackGraphSchema.h b/Engine/Plugins/Experimental/Animation/PoseSearch/Source/Editor/Public/AnimationBlendStackGraphSchema.h
deleted file mode 100644
diff --git a/Engine/Plugins/Experimental/Animation/PoseSearch/Source/Runtime/PoseSearch.Build.cs b/Engine/Plugins/Experimental/Animation/PoseSearch/Source/Runtime/PoseSearch.Build.cs
--- a/Engine/Plugins/Experimental/Animation/PoseSearch/Source/Runtime/PoseSearch.Build.cs
+++ b/Engine/Plugins/Experimental/Animation/PoseSearch/Source/Runtime/PoseSearch.Build.cs
@@ -22,7 +22,8 @@ public class PoseSearch : ModuleRules
 				"AnimationCore",
 				"TraceLog",
 				"AnimGraphRuntime",
-				"StructUtils"
+				"StructUtils",
+				"BlendStack"
 			}
 		);
 
diff --git a/Engine/Plugins/Experimental/Animation/PoseSearch/Source/Runtime/Private/AnimNode_BlendStack.cpp b/Engine/Plugins/Experimental/Animation/PoseSearch/Source/Runtime/Private/AnimNode_BlendStack.cpp
deleted file mode 100644
diff --git a/Engine/Plugins/Experimental/Animation/PoseSearch/Source/Runtime/Private/AnimNode_BlendStackInput.cpp b/Engine/Plugins/Experimental/Animation/PoseSearch/Source/Runtime/Private/AnimNode_BlendStackInput.cpp
deleted file mode 100644
diff --git a/Engine/Plugins/Experimental/Animation/PoseSearch/Source/Runtime/Private/AnimNode_MotionMatching.cpp b/Engine/Plugins/Experimental/Animation/PoseSearch/Source/Runtime/Private/AnimNode_MotionMatching.cpp
--- a/Engine/Plugins/Experimental/Animation/PoseSearch/Source/Runtime/Private/AnimNode_MotionMatching.cpp
+++ b/Engine/Plugins/Experimental/Animation/PoseSearch/Source/Runtime/Private/AnimNode_MotionMatching.cpp
@@ -143,7 +143,7 @@ void FAnimNode_MotionMatching::UpdateAssetPlayer(const FAnimationUpdateContext&
 				DebugInfo += FString::Printf(TEXT("  %s\n"), *GetNameSafe(DatabaseToSearch));
 			}
 			DebugInfo += FString::Printf(TEXT("Blend Stack:\n"));
-			for (const FPoseSearchAnimPlayer& AnimPlayer : AnimPlayers)
+			for (const FBlendStackAnimPlayer& AnimPlayer : AnimPlayers)
 			{
 				DebugInfo += FString::Printf(TEXT("  %s [time:%.2f|playrate:%.2f]\n"), *GetNameSafe(AnimPlayer.GetAnimationAsset()), AnimPlayer.GetAccumulatedTime(), AnimPlayer.GetPlayRate());
 			}
diff --git a/Engine/Plugins/Experimental/Animation/PoseSearch/Source/Runtime/Private/BlendStackAnimNodeLibrary.cpp b/Engine/Plugins/Experimental/Animation/PoseSearch/Source/Runtime/Private/BlendStackAnimNodeLibrary.cpp
deleted file mode 100644
diff --git a/Engine/Plugins/Experimental/Animation/PoseSearch/Source/Runtime/Public/PoseSearch/AnimNode_BlendStack.h b/Engine/Plugins/Experimental/Animation/PoseSearch/Source/Runtime/Public/PoseSearch/AnimNode_BlendStack.h
deleted file mode 100644
diff --git a/Engine/Plugins/Experimental/Animation/PoseSearch/Source/Runtime/Public/PoseSearch/AnimNode_BlendStackInput.h b/Engine/Plugins/Experimental/Animation/PoseSearch/Source/Runtime/Public/PoseSearch/AnimNode_BlendStackInput.h
deleted file mode 100644
diff --git a/Engine/Plugins/Experimental/Animation/PoseSearch/Source/Runtime/Public/PoseSearch/AnimNode_MotionMatching.h b/Engine/Plugins/Experimental/Animation/PoseSearch/Source/Runtime/Public/PoseSearch/AnimNode_MotionMatching.h
--- a/Engine/Plugins/Experimental/Animation/PoseSearch/Source/Runtime/Public/PoseSearch/AnimNode_MotionMatching.h
+++ b/Engine/Plugins/Experimental/Animation/PoseSearch/Source/Runtime/Public/PoseSearch/AnimNode_MotionMatching.h
@@ -4,7 +4,7 @@
 
 #include "Animation/AnimNode_AssetPlayerBase.h"
 #include "GameplayTagContainer.h"
-#include "PoseSearch/AnimNode_BlendStack.h"
+#include "BlendStack/AnimNode_BlendStack.h"
 #include "PoseSearch/PoseSearchLibrary.h"
 #include "PoseSearch/PoseSearchTrajectoryTypes.h"
 #include "AnimNode_MotionMatching.generated.h"
diff --git a/Engine/Plugins/Experimental/Animation/PoseSearch/Source/Runtime/Public/PoseSearch/BlendStackAnimNodeLibrary.h b/Engine/Plugins/Experimental/Animation/PoseSearch/Source/Runtime/Public/PoseSearch/BlendStackAnimNodeLibrary.h
deleted file mode 100644
