blendstack support for automatic transitions in state machine

diff --git a/Engine/Plugins/Experimental/Animation/BlendStack/Source/Runtime/Private/AnimNode_BlendStack.cpp b/Engine/Plugins/Experimental/Animation/BlendStack/Source/Runtime/Private/AnimNode_BlendStack.cpp
--- a/Engine/Plugins/Experimental/Animation/BlendStack/Source/Runtime/Private/AnimNode_BlendStack.cpp
+++ b/Engine/Plugins/Experimental/Animation/BlendStack/Source/Runtime/Private/AnimNode_BlendStack.cpp
@@ -294,7 +294,7 @@ FString FBlendStackAnimPlayer::GetAnimationName() const
 	return FString("StoredPose");
 }
 
-const UAnimationAsset* FBlendStackAnimPlayer::GetAnimationAsset() const
+UAnimationAsset* FBlendStackAnimPlayer::GetAnimationAsset() const
 {
 	if (SequencePlayerNode.GetSequence())
 	{
@@ -632,6 +632,11 @@ void FAnimNode_BlendStack_Standalone::InitializeSample(const FAnimationInitializ
 	}
 }
 
+UAnimationAsset* FAnimNode_BlendStack_Standalone::GetAnimAsset() const
+{
+	return AnimPlayers.IsEmpty() ? nullptr : AnimPlayers.First().GetAnimationAsset();
+}
+
 float FAnimNode_BlendStack_Standalone::GetAccumulatedTime() const
 {
 	return AnimPlayers.IsEmpty() ? 0.f : AnimPlayers.First().GetAccumulatedTime();
diff --git a/Engine/Plugins/Experimental/Animation/BlendStack/Source/Runtime/Public/BlendStack/AnimNode_BlendStack.h b/Engine/Plugins/Experimental/Animation/BlendStack/Source/Runtime/Public/BlendStack/AnimNode_BlendStack.h
--- a/Engine/Plugins/Experimental/Animation/BlendStack/Source/Runtime/Public/BlendStack/AnimNode_BlendStack.h
+++ b/Engine/Plugins/Experimental/Animation/BlendStack/Source/Runtime/Public/BlendStack/AnimNode_BlendStack.h
@@ -34,7 +34,7 @@ struct BLENDSTACK_API FBlendStackAnimPlayer
 	bool GetMirror() const { return MirrorNode.GetMirror(); }
 	FVector GetBlendParameters() const;
 	FString GetAnimationName() const;
-	const UAnimationAsset* GetAnimationAsset() const;
+	UAnimationAsset* GetAnimationAsset() const;
 
 	FAnimNode_Mirror_Standalone& GetMirrorNode() { return MirrorNode; }
 	int32 GetPoseLinkIndex() const { return PoseLinkIndex; }
@@ -110,6 +110,7 @@ struct BLENDSTACK_API FAnimNode_BlendStack_Standalone : public FAnimNode_AssetPl
 	void Reset();
 
 	// FAnimNode_AssetPlayerBase interface
+	virtual UAnimationAsset* GetAnimAsset() const;
 	virtual float GetAccumulatedTime() const override;
 	virtual void UpdateAssetPlayer(const FAnimationUpdateContext& Context) override;
 	// End of FAnimNode_AssetPlayerBase interface
