Add support for overriding PlayRate from within blend stack subgraph

diff --git a/Engine/Plugins/Experimental/Animation/BlendStack/Source/Editor/Private/AnimGraphNode_BlendStackResult.cpp b/Engine/Plugins/Experimental/Animation/BlendStack/Source/Editor/Private/AnimGraphNode_BlendStackResult.cpp
--- a/Engine/Plugins/Experimental/Animation/BlendStack/Source/Editor/Private/AnimGraphNode_BlendStackResult.cpp
+++ b/Engine/Plugins/Experimental/Animation/BlendStack/Source/Editor/Private/AnimGraphNode_BlendStackResult.cpp
@@ -7,10 +7,10 @@ bool UAnimGraphNode_BlendStackResult::IsNodeRootSet() const
 {
 	if (!Pins.IsEmpty() && !Pins[0]->LinkedTo.IsEmpty())
 	{
-		if (Pins[0]->LinkedTo[0]->GetOuter()->IsA<UAnimGraphNode_BlendStackInput>())
+		if (UAnimGraphNode_BlendStackInput* InputNode = Cast<UAnimGraphNode_BlendStackInput>(Pins[0]->LinkedTo[0]->GetOuter()))
 		{
-			// If output is connected to input directly, graph is a no-op. Consider it not set for pruning purposes.
-			return false;
+			// If output is connected to input directly, graph is a no-op, unless OverridePlayrate is true. Consider it not set for pruning purposes.
+			return InputNode->Node.bOverridePlayRate;
 		}
 	}
 
diff --git a/Engine/Plugins/Experimental/Animation/BlendStack/Source/Runtime/Private/AnimNode_BlendStackInput.cpp b/Engine/Plugins/Experimental/Animation/BlendStack/Source/Runtime/Private/AnimNode_BlendStackInput.cpp
--- a/Engine/Plugins/Experimental/Animation/BlendStack/Source/Runtime/Private/AnimNode_BlendStackInput.cpp
+++ b/Engine/Plugins/Experimental/Animation/BlendStack/Source/Runtime/Private/AnimNode_BlendStackInput.cpp
@@ -21,8 +21,11 @@ static FAnimNode_BlendStack_Standalone* GetBlendStackNodeFromIndex(const FAnimat
 	return LinkedProperty->ContainerPtrToValuePtr<FAnimNode_BlendStack_Standalone>(Context.AnimInstanceProxy->GetAnimInstanceObject());
 }
 
+
 void FAnimNode_BlendStackInput::Update_AnyThread(const FAnimationUpdateContext& Context)
 {
+	GetEvaluateGraphExposedInputs().Execute(Context);
+
 	// If there's no player, use the allocated blend stack index to get a reference to it.
 	// This should only happen on the first ever update of this node.
 	if (Player == nullptr)
@@ -36,7 +39,14 @@ void FAnimNode_BlendStackInput::Update_AnyThread(const FAnimationUpdateContext&
 		Player = &BlendStackNode->SampleGraphPoseLinks[SampleIndex].Player;
 	}
 
-	check(Player && *Player);
+	check(Player && *Player)
+
+	if (bOverridePlayRate)
+	{
+		(*Player)->UpdatePlayRate(PlayRate);
+	}
+	
 	(*Player)->Update_AnyThread(Context);
 }
 
+
diff --git a/Engine/Plugins/Experimental/Animation/BlendStack/Source/Runtime/Public/BlendStack/AnimNode_BlendStackInput.h b/Engine/Plugins/Experimental/Animation/BlendStack/Source/Runtime/Public/BlendStack/AnimNode_BlendStackInput.h
--- a/Engine/Plugins/Experimental/Animation/BlendStack/Source/Runtime/Public/BlendStack/AnimNode_BlendStackInput.h
+++ b/Engine/Plugins/Experimental/Animation/BlendStack/Source/Runtime/Public/BlendStack/AnimNode_BlendStackInput.h
@@ -17,13 +17,21 @@ public:
 
 	virtual void Update_AnyThread(const FAnimationUpdateContext& Context) override;
 	virtual void Evaluate_AnyThread(FPoseContext& Output) override;
-
+	
 	UPROPERTY()
 	int32 SampleIndex = INDEX_NONE;
 
 	UPROPERTY()
 	int32 BlendStackAllocationIndex = INDEX_NONE;
 
+	// If true, the PlayRate input from thos node will override the SequencePlayer or BlendSpacePlayer playrate each frame
+	UPROPERTY(EditAnywhere, Category=Settings, meta = (NeverAsPin))
+	bool bOverridePlayRate = false;
+	
+	// The play rate multiplier. Can be negative, which will cause the animation to play in reverse.
+	UPROPERTY(EditAnywhere, Category = Settings, meta = (PinHiddenByDefault))
+	float PlayRate = 0.0f;	
+
 	// The player is guaranteed to be valid for the whole duration of update/eval.
 	struct FBlendStackAnimPlayer** Player;
 };
