blend stack - MaxTimeBeforeFreezingInnerBlends

diff --git a/Engine/Plugins/Experimental/Animation/BlendStack/Source/Runtime/Private/AnimNode_BlendStack.cpp b/Engine/Plugins/Experimental/Animation/BlendStack/Source/Runtime/Private/AnimNode_BlendStack.cpp
--- a/Engine/Plugins/Experimental/Animation/BlendStack/Source/Runtime/Private/AnimNode_BlendStack.cpp
+++ b/Engine/Plugins/Experimental/Animation/BlendStack/Source/Runtime/Private/AnimNode_BlendStack.cpp
@@ -24,8 +24,10 @@ static constexpr FBoneIndexType RootBoneIndexType = 0;
 
 /////////////////////////////////////////////////////
 // FBlendStackAnimPlayer
-void FBlendStackAnimPlayer::Initialize(const FAnimationInitializeContext& Context, UAnimationAsset* AnimationAsset, float AccumulatedTime, bool bLoop, bool bMirrored, UMirrorDataTable* MirrorDataTable, float BlendTime, float RootBoneBlendTime, const UBlendProfile* BlendProfile, EAlphaBlendOption InBlendOption, FVector BlendParameters, float PlayRate, int32 InPoseLinkIdx,
-										FName GroupName, EAnimGroupRole::Type GroupRole, EAnimSyncMethod GroupMethod)
+void FBlendStackAnimPlayer::Initialize(const FAnimationInitializeContext& Context, UAnimationAsset* AnimationAsset, float AccumulatedTime, bool bLoop,
+	bool bMirrored, UMirrorDataTable* MirrorDataTable, float BlendTime, float RootBoneBlendTime, float MaxTimeBeforeFreezingInnerBlends, 
+	const UBlendProfile* BlendProfile, EAlphaBlendOption InBlendOption, FVector BlendParameters, float PlayRate, int32 InPoseLinkIdx,
+	FName GroupName, EAnimGroupRole::Type GroupRole, EAnimSyncMethod GroupMethod)
 {
 	if (bMirrored && !MirrorDataTable)
 	{
@@ -79,7 +81,9 @@ void FBlendStackAnimPlayer::Initialize(const FAnimationInitializeContext& Contex
 	BlendOption = InBlendOption;
 
 	TotalBlendInTime = BlendTime;
-	CurrentBlendInTime = 0.f;
+	MaxBlendInTimeAsSecondaryPlayer = MaxTimeBeforeFreezingInnerBlends >= 0 ? MaxTimeBeforeFreezingInnerBlends : BlendTime;
+	CurrentBlendInTimeAsMainPlayer = 0.f;
+	CurrentBlendInTimeAsSecondaryPlayer = 0.f;
 
 	MirrorNode.SetMirrorDataTable(MirrorDataTable);
 	MirrorNode.SetMirror(bMirrored);
@@ -267,6 +271,18 @@ float FBlendStackAnimPlayer::GetPlayRate() const
 	return 0.f;
 }
 
+void FBlendStackAnimPlayer::AdvanceBlendInTime(const float DeltaTime, bool bIsMainPlayer)
+{
+	if (bIsMainPlayer)
+	{
+		CurrentBlendInTimeAsMainPlayer += DeltaTime;
+	}
+	else
+	{
+		CurrentBlendInTimeAsSecondaryPlayer = FMath::Min(CurrentBlendInTimeAsSecondaryPlayer + DeltaTime, MaxBlendInTimeAsSecondaryPlayer);
+	}
+}
+
 FVector FBlendStackAnimPlayer::GetBlendParameters() const
 {
 	if (BlendSpacePlayerNode.GetBlendSpace())
@@ -316,7 +332,7 @@ float FBlendStackAnimPlayer::GetBlendInPercentage() const
 		return 1.f;
 	}
 
-	return FMath::Clamp(CurrentBlendInTime / TotalBlendInTime, 0.f, 1.f);
+	return FMath::Clamp(GetCurrentBlendInTime() / TotalBlendInTime, 0.f, 1.f);
 }
 
 bool FBlendStackAnimPlayer::GetBlendInWeights(TArray<float>& Weights) const
@@ -334,7 +350,7 @@ bool FBlendStackAnimPlayer::GetBlendInWeights(TArray<float>& Weights) const
 			}
 			else
 			{
-				const float UnclampedLinearWeight = CurrentBlendInTime / TotalBlendInTimeBoneIdx;
+				const float UnclampedLinearWeight = GetCurrentBlendInTime() / TotalBlendInTimeBoneIdx;
 				Weights[BoneIdx] = FAlphaBlend::AlphaToBlendOption(UnclampedLinearWeight, BlendOption);
 			}
 		}
@@ -607,7 +623,7 @@ void FAnimNode_BlendStack_Standalone::UpdateSample(const FAnimationUpdateContext
 	}
 
 	// Advance the blend-in time regardless of whether or not the player was updated.
-	SamplePlayer.AdvanceBlendInTime(Context.GetDeltaTime());
+	SamplePlayer.AdvanceBlendInTime(Context.GetDeltaTime(), PlayerIndex == 0);
 }
 
 void FAnimNode_BlendStack_Standalone::CacheBonesForSample(const FAnimationCacheBonesContext& Context, const int32 PlayerIndex)
@@ -668,8 +684,11 @@ static void RequestInertialBlend(const FAnimationUpdateContext& Context, float B
 	}
 }
 
-void FAnimNode_BlendStack_Standalone::BlendTo(const FAnimationUpdateContext& Context, UAnimationAsset* AnimationAsset, float AccumulatedTime, bool bLoop, bool bMirrored, UMirrorDataTable* MirrorDataTable, float BlendTime, float RootBoneBlendTime, const UBlendProfile* BlendProfile, EAlphaBlendOption BlendOption, bool bUseInertialBlend, FVector BlendParameters, float PlayRate,
-											FName GroupName, EAnimGroupRole::Type GroupRole, EAnimSyncMethod GroupMethod)
+void FAnimNode_BlendStack_Standalone::BlendTo(const FAnimationUpdateContext& Context, UAnimationAsset* AnimationAsset, float AccumulatedTime, bool bLoop,
+	bool bMirrored, UMirrorDataTable* MirrorDataTable,
+	float BlendTime, float RootBoneBlendTime, float MaxTimeBeforeFreezingInnerBlends, const UBlendProfile* BlendProfile, EAlphaBlendOption BlendOption,
+	bool bUseInertialBlend, FVector BlendParameters, float PlayRate,
+	FName GroupName, EAnimGroupRole::Type GroupRole, EAnimSyncMethod GroupMethod)
 {
 	if (bUseInertialBlend)
 	{
@@ -681,7 +700,7 @@ void FAnimNode_BlendStack_Standalone::BlendTo(const FAnimationUpdateContext& Con
 	FBlendStackAnimPlayer& AnimPlayer = AnimPlayers.First();
 
 	FAnimationInitializeContext InitContext(Context.AnimInstanceProxy, Context.SharedContext);
-	AnimPlayer.Initialize(InitContext, AnimationAsset, AccumulatedTime, bLoop, bMirrored, MirrorDataTable, BlendTime, RootBoneBlendTime, BlendProfile, BlendOption, BlendParameters, PlayRate, GetNextPoseLinkIndex(), GroupName, GroupRole, GroupMethod);
+	AnimPlayer.Initialize(InitContext, AnimationAsset, AccumulatedTime, bLoop, bMirrored, MirrorDataTable, BlendTime, RootBoneBlendTime, MaxTimeBeforeFreezingInnerBlends, BlendProfile, BlendOption, BlendParameters, PlayRate, GetNextPoseLinkIndex(), GroupName, GroupRole, GroupMethod);
 	InitializeSample(InitContext, AnimPlayer);
 }
 
@@ -786,7 +805,9 @@ void FAnimNode_BlendStack::UpdateAssetPlayer(const FAnimationUpdateContext& Cont
 
 	if (bExecuteBlendTo)
 	{
-		BlendTo(Context, AnimationAsset, AnimationTime, bLoop, bMirrored, MirrorDataTable.Get(), BlendTime, RootBoneBlendTime, BlendProfile, BlendOption, bUseInertialBlend, BlendParameters, WantedPlayRate);
+		BlendTo(Context, AnimationAsset, AnimationTime, bLoop, bMirrored, MirrorDataTable.Get(),
+			BlendTime, RootBoneBlendTime, MaxTimeBeforeFreezingInnerBlends, 
+			BlendProfile, BlendOption, bUseInertialBlend, BlendParameters, WantedPlayRate);
 	}
 	
 	UpdatePlayRate(WantedPlayRate);
diff --git a/Engine/Plugins/Experimental/Animation/BlendStack/Source/Runtime/Private/BlendStackAnimNodeLibrary.cpp b/Engine/Plugins/Experimental/Animation/BlendStack/Source/Runtime/Private/BlendStackAnimNodeLibrary.cpp
--- a/Engine/Plugins/Experimental/Animation/BlendStack/Source/Runtime/Private/BlendStackAnimNodeLibrary.cpp
+++ b/Engine/Plugins/Experimental/Animation/BlendStack/Source/Runtime/Private/BlendStackAnimNodeLibrary.cpp
@@ -38,6 +38,7 @@ void UBlendStackAnimNodeLibrary::BlendTo(const FAnimUpdateContext& Context,
 				BlendStackNodePtr->MirrorDataTable,
 				BlendTime,
 				BlendStackNodePtr->RootBoneBlendTime,
+				BlendStackNodePtr->MaxTimeBeforeFreezingInnerBlends,
 				BlendStackNodePtr->BlendProfile,
 				BlendStackNodePtr->BlendOption,
 				BlendStackNodePtr->bUseInertialBlend,
diff --git a/Engine/Plugins/Experimental/Animation/BlendStack/Source/Runtime/Public/BlendStack/AnimNode_BlendStack.h b/Engine/Plugins/Experimental/Animation/BlendStack/Source/Runtime/Public/BlendStack/AnimNode_BlendStack.h
--- a/Engine/Plugins/Experimental/Animation/BlendStack/Source/Runtime/Public/BlendStack/AnimNode_BlendStack.h
+++ b/Engine/Plugins/Experimental/Animation/BlendStack/Source/Runtime/Public/BlendStack/AnimNode_BlendStack.h
@@ -14,8 +14,10 @@ struct BLENDSTACK_API FBlendStackAnimPlayer
 {
 	GENERATED_BODY()
 	
-	void Initialize(const FAnimationInitializeContext& Context, UAnimationAsset* AnimationAsset, float AccumulatedTime, bool bLoop, bool bMirrored, UMirrorDataTable* MirrorDataTable, float BlendTime, float RootBoneBlendTime, const UBlendProfile* BlendProfile, EAlphaBlendOption InBlendOption, FVector BlendParameters, float PlayRate, int32 InPoseLinkIdx,
-					FName GroupName, EAnimGroupRole::Type GroupRole, EAnimSyncMethod Method);
+	void Initialize(const FAnimationInitializeContext& Context, UAnimationAsset* AnimationAsset, float AccumulatedTime, bool bLoop,
+		bool bMirrored, UMirrorDataTable* MirrorDataTable, float BlendTime, float RootBoneBlendTime, float MaxTimeBeforeFreezingInnerBlends,
+		const UBlendProfile* BlendProfile, EAlphaBlendOption InBlendOption, FVector BlendParameters, float PlayRate, int32 InPoseLinkIdx,
+		FName GroupName, EAnimGroupRole::Type GroupRole, EAnimSyncMethod Method);
 	
 	void UpdatePlayRate(float PlayRate);
 	void Evaluate_AnyThread(FPoseContext& Output);
@@ -29,8 +31,8 @@ struct BLENDSTACK_API FBlendStackAnimPlayer
 	void StorePoseContext(const FPoseContext& PoseContext);
 
 	float GetTotalBlendInTime() const { return TotalBlendInTime; }
-	float GetCurrentBlendInTime() const { return CurrentBlendInTime; }
-	void AdvanceBlendInTime(const float DeltaTime) { CurrentBlendInTime += DeltaTime; }
+	float GetCurrentBlendInTime() const { return CurrentBlendInTimeAsMainPlayer + CurrentBlendInTimeAsSecondaryPlayer; }
+	void AdvanceBlendInTime(const float DeltaTime, bool bIsMainPlayer);
 	bool GetMirror() const { return MirrorNode.GetMirror(); }
 	FVector GetBlendParameters() const;
 	FString GetAnimationName() const;
@@ -66,7 +68,9 @@ public:
 	TCustomBoneIndexArray<float, FSkeletonPoseBoneIndex> TotalBlendInTimePerBone;
 
 	float TotalBlendInTime = 0.f;
-	float CurrentBlendInTime = 0.f;
+	float MaxBlendInTimeAsSecondaryPlayer = UE_BIG_NUMBER;
+	float CurrentBlendInTimeAsMainPlayer = 0.f;
+	float CurrentBlendInTimeAsSecondaryPlayer = 0.f;
 };
 
 USTRUCT(BlueprintInternalUseOnly)
@@ -104,8 +108,12 @@ struct BLENDSTACK_API FAnimNode_BlendStack_Standalone : public FAnimNode_AssetPl
 	virtual void GatherDebugData(FNodeDebugData& DebugData) override;
 	// End of FAnimNode_Base interface
 
-	void BlendTo(const FAnimationUpdateContext& Context, UAnimationAsset* AnimationAsset, float AccumulatedTime = 0.f, bool bLoop = false, bool bMirrored = false, UMirrorDataTable* MirrorDataTable = nullptr, float BlendTime = 0.2f, float RootBoneBlendTime = -1.f, const UBlendProfile* BlendProfile = nullptr, EAlphaBlendOption BlendOption = EAlphaBlendOption::Linear, bool bUseInertialBlend = false, FVector BlendParameters = FVector::Zero(), float PlayRate = 1.f,
-					FName GroupName = NAME_None, EAnimGroupRole::Type GroupRole = EAnimGroupRole::CanBeLeader, EAnimSyncMethod Method = EAnimSyncMethod::DoNotSync);
+	void BlendTo(const FAnimationUpdateContext& Context, UAnimationAsset* AnimationAsset, float AccumulatedTime = 0.f, bool bLoop = false, 
+		bool bMirrored = false, UMirrorDataTable* MirrorDataTable = nullptr,
+		float BlendTime = 0.2f, float RootBoneBlendTime = -1.f, float MaxTimeBeforeFreezingInnerBlends = -1.f,
+		const UBlendProfile* BlendProfile = nullptr, EAlphaBlendOption BlendOption = EAlphaBlendOption::Linear, 
+		bool bUseInertialBlend = false, FVector BlendParameters = FVector::Zero(), float PlayRate = 1.f,
+		FName GroupName = NAME_None, EAnimGroupRole::Type GroupRole = EAnimGroupRole::CanBeLeader, EAnimSyncMethod Method = EAnimSyncMethod::DoNotSync);
 	void UpdatePlayRate(float PlayRate);
 	void Reset();
 
@@ -169,6 +177,10 @@ struct BLENDSTACK_API FAnimNode_BlendStack : public FAnimNode_BlendStack_Standal
 	UPROPERTY(EditAnywhere, BlueprintReadWrite, Category = Settings, meta = (PinHiddenByDefault))
 	float RootBoneBlendTime = -1.f;
 
+	// Max time in seconds before freezing the blend ratio of all blends, but the most recently requested one. Negative values would imply MaxTimeBeforeFreezingInnerBlends is equal to BlendTime
+	UPROPERTY(EditAnywhere, BlueprintReadWrite, Category = Settings, meta = (PinHiddenByDefault))
+	float MaxTimeBeforeFreezingInnerBlends = -1.f;
+
 	// if AnimationTime and MaxAnimationDeltaTime are positive and the currently playing animation total time differs more than MaxAnimationDeltaTime from AnimationTime
 	// (animation desynchronized from the requested time) the blend stack will force a blend into the same animation
 	UPROPERTY(EditAnywhere, BlueprintReadWrite, Category = Settings, meta = (PinHiddenByDefault))
diff --git a/Engine/Plugins/Experimental/Animation/PoseSearch/Source/Runtime/Private/AnimNode_MotionMatching.cpp b/Engine/Plugins/Experimental/Animation/PoseSearch/Source/Runtime/Private/AnimNode_MotionMatching.cpp
--- a/Engine/Plugins/Experimental/Animation/PoseSearch/Source/Runtime/Private/AnimNode_MotionMatching.cpp
+++ b/Engine/Plugins/Experimental/Animation/PoseSearch/Source/Runtime/Private/AnimNode_MotionMatching.cpp
@@ -193,7 +193,8 @@ void FAnimNode_MotionMatching::UpdateAssetPlayer(const FAnimationUpdateContext&
 				const float RootBoneBlendTime = YawFromAnimationBlendRate < 0.f ? BlendTime : 0.f;
 				FAnimNode_BlendStack_Standalone::BlendTo(Context, DatabaseAsset->GetAnimationAsset(), MotionMatchingState.CurrentSearchResult.AssetTime,
 					DatabaseAsset->IsLooping(), SearchIndexAsset->bMirrored, CurrentResultDatabase->Schema->MirrorDataTable.Get(),
-					BlendTime, RootBoneBlendTime, BlendProfile, BlendOption, bUseInertialBlend, SearchIndexAsset->BlendParameters, MotionMatchingState.WantedPlayRate);
+					BlendTime, RootBoneBlendTime, MaxTimeBeforeFreezingInnerBlends,
+					BlendProfile, BlendOption, bUseInertialBlend, SearchIndexAsset->BlendParameters, MotionMatchingState.WantedPlayRate);
 			}
 		}
 	}
diff --git a/Engine/Plugins/Experimental/Animation/PoseSearch/Source/Runtime/Public/PoseSearch/AnimNode_MotionMatching.h b/Engine/Plugins/Experimental/Animation/PoseSearch/Source/Runtime/Public/PoseSearch/AnimNode_MotionMatching.h
--- a/Engine/Plugins/Experimental/Animation/PoseSearch/Source/Runtime/Public/PoseSearch/AnimNode_MotionMatching.h
+++ b/Engine/Plugins/Experimental/Animation/PoseSearch/Source/Runtime/Public/PoseSearch/AnimNode_MotionMatching.h
@@ -60,6 +60,10 @@ private:
 	UPROPERTY(EditAnywhere, Category = Settings, meta = (PinHiddenByDefault, ClampMin="0"))
 	float BlendTime = 0.2f;
 
+	// Max time in seconds before freezing the blend ratio of all blends, but the most recently requested one. Negative values would imply MaxTimeBeforeFreezingInnerBlends is equal to BlendTime
+	UPROPERTY(EditAnywhere, Category = Settings, meta = (PinHiddenByDefault))
+	float MaxTimeBeforeFreezingInnerBlends = -1.f;
+
 	// Set Blend Profiles (editable in the skeleton) to determine how the blending is distributed among your character's bones. It could be used to differentiate between upper body and lower body to blend timing.
 	UPROPERTY(EditAnywhere, Category = Settings, meta = (PinHiddenByDefault, UseAsBlendProfile = true))
 	TObjectPtr<UBlendProfile> BlendProfile;
diff --git a/Engine/Plugins/Experimental/Chooser/Source/Chooser/Private/AnimNode_ChooserPlayer.cpp b/Engine/Plugins/Experimental/Chooser/Source/Chooser/Private/AnimNode_ChooserPlayer.cpp
--- a/Engine/Plugins/Experimental/Chooser/Source/Chooser/Private/AnimNode_ChooserPlayer.cpp
+++ b/Engine/Plugins/Experimental/Chooser/Source/Chooser/Private/AnimNode_ChooserPlayer.cpp
@@ -117,7 +117,8 @@ void FAnimNode_ChooserPlayer::UpdateAssetPlayer(const FAnimationUpdateContext& C
 		
 		FAnimNode_BlendStack_Standalone::BlendTo(Context, NewAsset, Settings.StartTime,
 			bLoop, Settings.bMirror, MirrorDataTable,
-			Settings.BlendTime, Settings.BlendTime, Settings.BlendProfile, Settings.BlendOption, Settings.bUseInertialBlend, FVector::Zero(), Settings.PlaybackRate,
+			Settings.BlendTime, Settings.BlendTime, Settings.BlendTime,
+			Settings.BlendProfile, Settings.BlendOption, Settings.bUseInertialBlend, FVector::Zero(), Settings.PlaybackRate,
 			GetGroupName(), GetGroupRole(), GetGroupMethod());
 		
 		CurrentAsset = NewAsset;
