blend stack - performance improvements

diff --git a/Engine/Plugins/Experimental/Animation/BlendStack/Source/Runtime/Private/AnimNode_BlendStack.cpp b/Engine/Plugins/Experimental/Animation/BlendStack/Source/Runtime/Private/AnimNode_BlendStack.cpp
--- a/Engine/Plugins/Experimental/Animation/BlendStack/Source/Runtime/Private/AnimNode_BlendStack.cpp
+++ b/Engine/Plugins/Experimental/Animation/BlendStack/Source/Runtime/Private/AnimNode_BlendStack.cpp
@@ -14,8 +14,8 @@
 #include UE_INLINE_GENERATED_CPP_BY_NAME(AnimNode_BlendStack)
 
 #if ENABLE_ANIM_DEBUG
-TAutoConsoleVariable<int32> CVarAnimBlendStackEnable(TEXT("a.AnimNode.BlendStack.Enable"), 1, TEXT("Enable / Disable Blend Stack"));
-TAutoConsoleVariable<int32> CVarAnimBlendStackPruningEnable(TEXT("a.AnimNode.BlendStack.Pruning.Enable"), 1, TEXT("Enable / Disable Blend Stack Pruning"));
+static TAutoConsoleVariable<bool> CVarAnimBlendStackEnable(TEXT("a.AnimNode.BlendStack.Enable"), true, TEXT("Enable / Disable Blend Stack"));
+static TAutoConsoleVariable<bool> CVarAnimBlendStackPruningEnable(TEXT("a.AnimNode.BlendStack.Pruning.Enable"), true, TEXT("Enable / Disable Blend Stack Pruning"));
 #endif
 
 #define LOCTEXT_NAMESPACE "AnimNode_BlendStack"
@@ -340,28 +340,28 @@ float FBlendStackAnimPlayer::GetBlendInPercentage() const
 	return FMath::Clamp(GetCurrentBlendInTime() / TotalBlendInTime, 0.f, 1.f);
 }
 
-bool FBlendStackAnimPlayer::GetBlendInWeights(TArray<float>& Weights) const
+int32 FBlendStackAnimPlayer::GetBlendInWeightsNum() const
 {
-	const int32 NumBones = TotalBlendInTimePerBone.Num();
-	if (NumBones > 0)
+	return TotalBlendInTimePerBone.Num();
+}
+
+void FBlendStackAnimPlayer::GetBlendInWeights(TArrayView<float> Weights) const
+{
+	check(Weights.Num() == GetBlendInWeightsNum());
+	
+	for (int32 BoneIdx = 0; BoneIdx < Weights.Num(); ++BoneIdx)
 	{
-		Weights.SetNumUninitialized(NumBones);
-		for (int32 BoneIdx = 0; BoneIdx < NumBones; ++BoneIdx)
+		const float TotalBlendInTimeBoneIdx = TotalBlendInTimePerBone[BoneIdx];
+		if (FMath::IsNearlyZero(TotalBlendInTimeBoneIdx))
 		{
-			const float TotalBlendInTimeBoneIdx = TotalBlendInTimePerBone[BoneIdx];
-			if (FMath::IsNearlyZero(TotalBlendInTimeBoneIdx))
-			{
-				Weights[BoneIdx] = 1.f;
-			}
-			else
-			{
-				const float UnclampedLinearWeight = GetCurrentBlendInTime() / TotalBlendInTimeBoneIdx;
-				Weights[BoneIdx] = FAlphaBlend::AlphaToBlendOption(UnclampedLinearWeight, BlendOption);
-			}
+			Weights[BoneIdx] = 1.f;
+		}
+		else
+		{
+			const float UnclampedLinearWeight = GetCurrentBlendInTime() / TotalBlendInTimeBoneIdx;
+			Weights[BoneIdx] = FAlphaBlend::AlphaToBlendOption(UnclampedLinearWeight, BlendOption);
 		}
-		return true;
 	}
-	return false;
 }
 
 /////////////////////////////////////////////////////
@@ -401,37 +401,39 @@ void FAnimNode_BlendStack_Standalone::Evaluate_AnyThread(FPoseContext& Output)
 		EvaluateSample(Output, BlendStackSize - 1);
 
 		FPoseContext EvaluationPoseContext(Output);
-		FPoseContext BlendedPoseContext(Output); // @todo: this should not be necessary (but FBaseBlendedCurve::InitFrom complains about "ensure(&InCurveToInitFrom != this)"): optimize it away!
-		FAnimationPoseData BlendedAnimationPoseData(BlendedPoseContext);
 
 		const USkeleton* SkeletonAsset = Output.AnimInstanceProxy->GetRequiredBones().GetSkeletonAsset();
 		check(SkeletonAsset);
 
 		const FReferenceSkeleton& RefSkeleton = SkeletonAsset->GetReferenceSkeleton();
 		const int32 NumSkeletonBones = RefSkeleton.GetNum();
-		TArray<float> Weights;
 
-		auto EvaluateAndBlendPlayerByIndex = [this, &EvaluationPoseContext, &BlendedAnimationPoseData, &Output, &Weights, &BlendedPoseContext](int32 PlayerIndex)
+		auto EvaluateAndBlendPlayerByIndex = [this, &EvaluationPoseContext, &Output](int32 PlayerIndex)
 		{
+			FAnimationPoseData OutputAnimationPoseData(Output);
+			FAnimationPoseData EvaluationAnimationPoseData(EvaluationPoseContext);
+
 			// Evaluate into EvaluationPoseContext and then blend it with the Output (initialized with the last AnimPlayer evaluation)
 			EvaluateSample(EvaluationPoseContext, PlayerIndex);
-			if (AnimPlayers[PlayerIndex].GetBlendInWeights(Weights))
+
+			const int32 BlendInWeightsNum = AnimPlayers[PlayerIndex].GetBlendInWeightsNum();
+			if (BlendInWeightsNum > 0)
 			{
-				// @todo: have BlendTwoPosesTogetherPerBone using a TArrayView for the Weights to avoid allocations
-				FAnimationRuntime::BlendTwoPosesTogetherPerBone(FAnimationPoseData(Output), FAnimationPoseData(EvaluationPoseContext), Weights, BlendedAnimationPoseData);
+				TArrayView<float> Weights((float*)FMemory_Alloca(BlendInWeightsNum * sizeof(float)), BlendInWeightsNum);
+				AnimPlayers[PlayerIndex].GetBlendInWeights(Weights);
+				BlendWithPosePerBone(OutputAnimationPoseData, EvaluationAnimationPoseData, Weights);
 			}
 			else
 			{
 				const float Weight = 1.f - FAlphaBlend::AlphaToBlendOption(AnimPlayers[PlayerIndex].GetBlendInPercentage(), AnimPlayers[PlayerIndex].GetBlendOption());
-				FAnimationRuntime::BlendTwoPosesTogether(FAnimationPoseData(Output), FAnimationPoseData(EvaluationPoseContext), Weight, BlendedAnimationPoseData);
+				BlendWithPose(OutputAnimationPoseData, EvaluationAnimationPoseData, Weight);
 			}
-			Output = BlendedPoseContext; // @todo: this should not be necessary either: optimize it away!
 		};
 
 #if ENABLE_ANIM_DEBUG
-			const bool bEnablePruning = CVarAnimBlendStackPruningEnable.GetValueOnAnyThread() > 0;
+		const bool bEnablePruning = CVarAnimBlendStackPruningEnable.GetValueOnAnyThread();
 #else
-			const bool bEnablePruning = true;
+		const bool bEnablePruning = true;
 #endif // ENABLE_ANIM_DEBUG
 
 		// Evaluate our players from the second last to the first.
@@ -756,6 +758,86 @@ void FAnimNode_BlendStack_Standalone::GatherDebugData(FNodeDebugData& DebugData)
 	}
 }
 
+// this function is the optimized version of 
+// FAnimationRuntime::BlendTwoPosesTogetherPerBone(InOutPoseData, OtherPoseData, OtherPoseWeights, InOutPoseData);
+void FAnimNode_BlendStack_Standalone::BlendWithPosePerBone(FAnimationPoseData& InOutPoseData, const FAnimationPoseData& OtherPoseData, TConstArrayView<float> OtherPoseWeights)
+{
+	using namespace UE::Anim;
+
+	FCompactPose& InOutPose = InOutPoseData.GetPose();
+	FBlendedCurve& InOutCurve = InOutPoseData.GetCurve();
+	FStackAttributeContainer& InOutAttributes = InOutPoseData.GetAttributes();
+
+	const FCompactPose& OtherPose = OtherPoseData.GetPose();
+
+	for (FCompactPoseBoneIndex BoneIndex : InOutPose.ForEachBoneIndex())
+	{
+		const float OtherPoseBoneWeight = OtherPoseWeights[BoneIndex.GetInt()];
+		if (FAnimationRuntime::IsFullWeight(OtherPoseBoneWeight))
+		{
+			InOutPose[BoneIndex] = OtherPose[BoneIndex];
+		}
+		else if (FAnimationRuntime::HasWeight(OtherPoseBoneWeight))
+		{
+			const ScalarRegister VInOutPoseBoneWeight(1.f - OtherPoseBoneWeight);
+			const ScalarRegister VOtherPoseBoneWeight(OtherPoseBoneWeight);
+
+			InOutPose[BoneIndex] *= VInOutPoseBoneWeight;
+			InOutPose[BoneIndex].AccumulateWithShortestRotation(OtherPose[BoneIndex], VOtherPoseBoneWeight);
+		}
+		// else we leave InOutPose[BoneIndex] as is
+	}
+
+	// Ensure that all of the resulting rotations are normalized
+	InOutPose.NormalizeRotations();
+
+	// @note : This isn't perfect as curve can link to joint, and it would be the best to use that information
+	// but that is very expensive option as we have to have another indirect look up table to search. 
+	// For now, replacing with combine (non-zero will be overriden)
+	// in the future, we might want to do this outside if we want per bone blend to apply curve also UE-39182
+	InOutPoseData.GetCurve().Combine(OtherPoseData.GetCurve());
+
+	// @todo: optimize away the copy
+	FStackAttributeContainer CustomAttributes;
+	Attributes::BlendAttributesPerBone(InOutAttributes, OtherPoseData.GetAttributes(), OtherPoseWeights, CustomAttributes);
+	InOutAttributes = CustomAttributes;
+}
+
+// this function is the optimized version of 
+// FAnimationRuntime::BlendTwoPosesTogether(InOutPoseData, OtherPoseData, InOutPoseWeight, InOutPoseData);
+void FAnimNode_BlendStack_Standalone::BlendWithPose(FAnimationPoseData& InOutPoseData, const FAnimationPoseData& OtherPoseData, const float InOutPoseWeight)
+{
+	using namespace UE::Anim;
+
+	FCompactPose& InOutPose = InOutPoseData.GetPose();
+	FBlendedCurve& InOutCurve = InOutPoseData.GetCurve();
+	FStackAttributeContainer& InOutAttributes = InOutPoseData.GetAttributes();
+
+	const FCompactPose& OtherPose = OtherPoseData.GetPose();
+
+	const float OtherPoseWeight = 1.f - InOutPoseWeight;
+
+	// @todo: reimplement the ispc version of this if needed
+	const ScalarRegister VInOutPoseWeight(InOutPoseWeight);
+	const ScalarRegister VOtherPoseWeight(OtherPoseWeight);
+
+	for (FCompactPoseBoneIndex BoneIndex : InOutPose.ForEachBoneIndex())
+	{
+		InOutPose[BoneIndex] *= VInOutPoseWeight;
+		InOutPose[BoneIndex].AccumulateWithShortestRotation(OtherPose[BoneIndex], VOtherPoseWeight);
+	}
+
+	// Ensure that all of the resulting rotations are normalized
+	InOutPose.NormalizeRotations();
+
+	InOutCurve.LerpTo(OtherPoseData.GetCurve(), OtherPoseWeight);
+	
+	// @todo: optimize away the copy
+	FStackAttributeContainer CustomAttributes;
+	Attributes::BlendAttributes({ InOutAttributes, OtherPoseData.GetAttributes() }, { InOutPoseWeight, OtherPoseWeight }, { 0, 1 }, CustomAttributes);
+	InOutAttributes = CustomAttributes;
+}
+
 /////////////////////////////////////////////////////
 // FAnimNode_BlendStack
 
diff --git a/Engine/Plugins/Experimental/Animation/BlendStack/Source/Runtime/Public/BlendStack/AnimNode_BlendStack.h b/Engine/Plugins/Experimental/Animation/BlendStack/Source/Runtime/Public/BlendStack/AnimNode_BlendStack.h
--- a/Engine/Plugins/Experimental/Animation/BlendStack/Source/Runtime/Public/BlendStack/AnimNode_BlendStack.h
+++ b/Engine/Plugins/Experimental/Animation/BlendStack/Source/Runtime/Public/BlendStack/AnimNode_BlendStack.h
@@ -26,7 +26,10 @@ struct BLENDSTACK_API FBlendStackAnimPlayer
 	float GetPlayRate() const;
 	
 	float GetBlendInPercentage() const;
-	bool GetBlendInWeights(TArray<float>& Weights) const;
+	
+	int32 GetBlendInWeightsNum() const;
+	void GetBlendInWeights(TArrayView<float> Weights) const;
+
 	EAlphaBlendOption GetBlendOption() const { return BlendOption; }
 	void StorePoseContext(const FPoseContext& PoseContext);
 
@@ -132,6 +135,9 @@ struct BLENDSTACK_API FAnimNode_BlendStack_Standalone : public FAnimNode_AssetPl
 	int32 GetMaxActiveBlends() const { return MaxActiveBlends; }
 
 protected:
+	static void BlendWithPose(FAnimationPoseData& InOutPoseData, const FAnimationPoseData& OtherPoseData, const float InOutPoseWeight);
+	static void BlendWithPosePerBone(FAnimationPoseData& InOutPoseData, const FAnimationPoseData& OtherPoseData, TConstArrayView<float> OtherPoseWeights);
+
 	// Number of max active blending animation in the blend stack. If MaxActiveBlends is zero then blend stack is disabled
 	UPROPERTY(EditAnywhere, Category = Settings, meta = (ClampMin = "0"))
 	int32 MaxActiveBlends = 4;
