blend stack - memory optimization

diff --git a/Engine/Plugins/Experimental/Animation/BlendStack/Source/Runtime/Private/AnimNode_BlendStack.cpp b/Engine/Plugins/Experimental/Animation/BlendStack/Source/Runtime/Private/AnimNode_BlendStack.cpp
--- a/Engine/Plugins/Experimental/Animation/BlendStack/Source/Runtime/Private/AnimNode_BlendStack.cpp
+++ b/Engine/Plugins/Experimental/Animation/BlendStack/Source/Runtime/Private/AnimNode_BlendStack.cpp
@@ -26,7 +26,7 @@ static constexpr FBoneIndexType RootBoneIndexType = 0;
 // FBlendStackAnimPlayer
 void FBlendStackAnimPlayer::Initialize(const FAnimationInitializeContext& Context, UAnimationAsset* AnimationAsset, float AccumulatedTime, bool bLoop,
 	bool bMirrored, UMirrorDataTable* MirrorDataTable, float BlendTime, float RootBoneBlendTime, float MaxTimeBeforeFreezingInnerBlends, 
-	const UBlendProfile* BlendProfile, EAlphaBlendOption InBlendOption, FVector BlendParameters, float PlayRate, int32 InPoseLinkIdx,
+	const UBlendProfile* BlendProfile, EAlphaBlendOption InBlendOption, const FVector& BlendParameters, float PlayRate, int32 InPoseLinkIdx,
 	FName GroupName, EAnimGroupRole::Type GroupRole, EAnimSyncMethod GroupMethod)
 {
 	if (bMirrored && !MirrorDataTable)
@@ -160,6 +160,19 @@ void FBlendStackAnimPlayer::StorePoseContext(const FPoseContext& PoseContext)
 	StoredAttributes.CopyFrom(PoseContext.CustomAttributes);
 }
 
+bool FBlendStackAnimPlayer::HasValidPoseContext() const
+{
+	return StoredPose.IsValid();
+}
+
+void FBlendStackAnimPlayer::MovePoseContextTo(FBlendStackAnimPlayer& Other)
+{
+	Other.StoredPose = MoveTemp(StoredPose);
+	Other.StoredCurve = MoveTemp(StoredCurve);
+	Other.StoredAttributes = MoveTemp(StoredAttributes);
+	Other.StoredBoneContainer = MoveTemp(StoredBoneContainer);
+}
+
 void FBlendStackAnimPlayer::RestorePoseContext(FPoseContext& PoseContext) const
 {
 	check(!SequencePlayerNode.GetSequence() && !BlendSpacePlayerNode.GetBlendSpace());
@@ -220,6 +233,21 @@ void FBlendStackAnimPlayer::UpdateSourceLinkNode()
 	}
 }
 
+bool FBlendStackAnimPlayer::IsLooping() const
+{
+	if (SequencePlayerNode.GetSequence())
+	{
+		return SequencePlayerNode.IsLooping();
+	}
+
+	if (BlendSpacePlayerNode.GetBlendSpace())
+	{
+		return BlendSpacePlayerNode.IsLooping();
+	}
+
+	return false;
+}
+
 void FBlendStackAnimPlayer::Evaluate_AnyThread(FPoseContext& Output)
 {
 	if (SequencePlayerNode.GetSequence() || BlendSpacePlayerNode.GetBlendSpace())
@@ -298,6 +326,14 @@ FVector FBlendStackAnimPlayer::GetBlendParameters() const
 	return FVector::ZeroVector;
 }
 
+void FBlendStackAnimPlayer::SetBlendParameters(const FVector& BlendParameters)
+{
+	if (BlendSpacePlayerNode.GetBlendSpace())
+	{
+		BlendSpacePlayerNode.SetPosition(BlendParameters);
+	}
+}
+
 FString FBlendStackAnimPlayer::GetAnimationName() const
 {
 	if (SequencePlayerNode.GetSequence())
@@ -366,6 +402,26 @@ void FBlendStackAnimPlayer::GetBlendInWeights(TArrayView<float> Weights) const
 
 /////////////////////////////////////////////////////
 // FAnimNode_BlendStack_Standalone
+void FAnimNode_BlendStack_Standalone::PopLastAnimPlayer()
+{
+	const int32 LastAnimPlayerIndex = AnimPlayers.Num() - 1;
+
+#if DO_CHECK
+	for (int32 AnimPlayerIndex = 0; AnimPlayerIndex < LastAnimPlayerIndex; ++AnimPlayerIndex)
+	{
+		// making sure only the last AnimPlayer can have a valid pose
+		check(!AnimPlayers[AnimPlayerIndex].HasValidPoseContext());
+	}
+#endif // DO_CHECK
+
+	if (LastAnimPlayerIndex > 0 && AnimPlayers[LastAnimPlayerIndex].HasValidPoseContext())
+	{
+		AnimPlayers[LastAnimPlayerIndex].MovePoseContextTo(AnimPlayers[LastAnimPlayerIndex - 1]);
+	}
+
+	AnimPlayers.PopLast();
+}
+
 void FAnimNode_BlendStack_Standalone::Evaluate_AnyThread(FPoseContext& Output)
 {
 	DECLARE_SCOPE_HIERARCHICAL_COUNTER_ANIMNODE(Evaluate_AnyThread);
@@ -391,7 +447,7 @@ void FAnimNode_BlendStack_Standalone::Evaluate_AnyThread(FPoseContext& Output)
 		// Disable blend stack if requested (for testing / debugging) by removing all the AnimPlayers except the first
 		while (AnimPlayers.Num() > 1)
 		{
-			AnimPlayers.PopLast();
+			PopLastAnimPlayer();
 		}
 		EvaluateSample(Output, 0);
 	}
@@ -436,7 +492,9 @@ void FAnimNode_BlendStack_Standalone::Evaluate_AnyThread(FPoseContext& Output)
 		const bool bEnablePruning = true;
 #endif // ENABLE_ANIM_DEBUG
 
-		// Evaluate our players from the second last to the first.
+		// ...continuing with the valuation and accumulation on the Output FPoseContext
+		// of AnimPlayer(s) from the second last to the AnimPlayer[MaxActiveBlends].
+		// Popping them if bEnablePruning
 		int32 PlayerIndex = BlendStackSize - 2;
 		// Start evaluating with our least significant players.
 		for (; PlayerIndex >= MaxActiveBlends; --PlayerIndex)
@@ -446,30 +504,34 @@ void FAnimNode_BlendStack_Standalone::Evaluate_AnyThread(FPoseContext& Output)
 			if (bEnablePruning)
 			{
 				// too many AnimPlayers! we don't have enough available blends to hold them all, so we accumulate the blended poses into Output / BlendedPoseContext.
-				AnimPlayers.PopLast();
+				PopLastAnimPlayer();
 			}
 		}
 
-		// Even if we're not pruning, we must use the stored pose if we have a limited number of graphs to execute.
-		const bool bNeedsStoredPose = (bEnablePruning || !SampleGraphPoseLinks.IsEmpty()) && (PlayerIndex == (MaxActiveBlends - 1));
-		if (bNeedsStoredPose)
+		// At this point Output FPoseContext contains all the weighted accumulated poses of the from AnimPlayer[MaxActiveBlends] to AnimPlayer[AnimPlayer.Num()-1]
+		if (PlayerIndex == (MaxActiveBlends - 1))
 		{
 			check(AnimPlayers.Num() == MaxActiveBlends + 1);
 
-			// We store Output / BlendedPoseContext into the last AnimPlayer, that will hold a static pose, no longer an animation playing.
-			AnimPlayers.Last().StorePoseContext(Output);
+			if (bEnablePruning && bStoreBlendedPose)
+			{
+				// We store Output / BlendedPoseContext into the last AnimPlayer, that will hold a static pose, no longer an animation playing.
+				AnimPlayers.Last().StorePoseContext(Output);
+			}
 
+			// we execute the associated graph on the Output
 			if (!SampleGraphPoseLinks.IsEmpty())
 			{
 				const int32 PoseLinkIdx = AnimPlayers[MaxActiveBlends].GetPoseLinkIndex();
 				FBlendStack_SampleGraphPoseLink& PoseLink = SampleGraphPoseLinks[PoseLinkIdx];
+
 				// No players should have evaluated a graph before this point.
 				// Evaluate the graph on the blended result.
 				PoseLink.EvaluatePlayer(Output, AnimPlayers[MaxActiveBlends]);
 			}
 		}
 
-		// Continue with our most significant players.
+		// Continue with the evaluation of the most significant AnimPlayer(s) with the associated graphs
 		for (; PlayerIndex >= 0; --PlayerIndex)
 		{
 			EvaluateAndBlendPlayerByIndex(PlayerIndex);
@@ -553,7 +615,7 @@ void FAnimNode_BlendStack_Standalone::UpdateAssetPlayer(const FAnimationUpdateCo
 	const int32 WantedAnimPlayersNum = FMath::Max(1, AnimPlayerIndex); // we save at least one FBlendStackAnimPlayer
 	while (AnimPlayers.Num() > WantedAnimPlayersNum)
 	{
-		AnimPlayers.PopLast();
+		PopLastAnimPlayer();
 	}
 }
 
@@ -694,7 +756,7 @@ static void RequestInertialBlend(const FAnimationUpdateContext& Context, float B
 void FAnimNode_BlendStack_Standalone::BlendTo(const FAnimationUpdateContext& Context, UAnimationAsset* AnimationAsset, float AccumulatedTime, bool bLoop,
 	bool bMirrored, UMirrorDataTable* MirrorDataTable,
 	float BlendTime, float RootBoneBlendTime, float MaxTimeBeforeFreezingInnerBlends, const UBlendProfile* BlendProfile, EAlphaBlendOption BlendOption,
-	bool bUseInertialBlend, FVector BlendParameters, float PlayRate,
+	bool bUseInertialBlend, const FVector& BlendParameters, float PlayRate,
 	FName GroupName, EAnimGroupRole::Type GroupRole, EAnimSyncMethod GroupMethod)
 {
 	if (bUseInertialBlend)
diff --git a/Engine/Plugins/Experimental/Animation/BlendStack/Source/Runtime/Public/BlendStack/AnimNode_BlendStack.h b/Engine/Plugins/Experimental/Animation/BlendStack/Source/Runtime/Public/BlendStack/AnimNode_BlendStack.h
--- a/Engine/Plugins/Experimental/Animation/BlendStack/Source/Runtime/Public/BlendStack/AnimNode_BlendStack.h
+++ b/Engine/Plugins/Experimental/Animation/BlendStack/Source/Runtime/Public/BlendStack/AnimNode_BlendStack.h
@@ -16,7 +16,7 @@ struct BLENDSTACK_API FBlendStackAnimPlayer
 	
 	void Initialize(const FAnimationInitializeContext& Context, UAnimationAsset* AnimationAsset, float AccumulatedTime, bool bLoop,
 		bool bMirrored, UMirrorDataTable* MirrorDataTable, float BlendTime, float RootBoneBlendTime, float MaxTimeBeforeFreezingInnerBlends,
-		const UBlendProfile* BlendProfile, EAlphaBlendOption InBlendOption, FVector BlendParameters, float PlayRate, int32 InPoseLinkIdx,
+		const UBlendProfile* BlendProfile, EAlphaBlendOption InBlendOption, const FVector& BlendParameters, float PlayRate, int32 InPoseLinkIdx,
 		FName GroupName, EAnimGroupRole::Type GroupRole, EAnimSyncMethod Method);
 	
 	void UpdatePlayRate(float PlayRate);
@@ -32,22 +32,29 @@ struct BLENDSTACK_API FBlendStackAnimPlayer
 
 	EAlphaBlendOption GetBlendOption() const { return BlendOption; }
 	void StorePoseContext(const FPoseContext& PoseContext);
+	bool HasValidPoseContext() const;
+	void MovePoseContextTo(FBlendStackAnimPlayer& Other);
 
 	float GetTotalBlendInTime() const { return TotalBlendInTime; }
 	float GetCurrentBlendInTime() const { return CurrentBlendInTimeAsMainPlayer + CurrentBlendInTimeAsSecondaryPlayer; }
 	void AdvanceBlendInTime(const float DeltaTime, bool bIsMainPlayer);
 	bool GetMirror() const { return MirrorNode.GetMirror(); }
 	FVector GetBlendParameters() const;
+	void SetBlendParameters(const FVector& BlendParameters);
 	FString GetAnimationName() const;
 	UAnimationAsset* GetAnimationAsset() const;
 
 	FAnimNode_Mirror_Standalone& GetMirrorNode() { return MirrorNode; }
 	int32 GetPoseLinkIndex() const { return PoseLinkIndex; }
 
-public:
 	void RestorePoseContext(FPoseContext& PoseContext) const;
 	void UpdateSourceLinkNode();
+	bool IsLooping() const;
+	
+	// Curves to add to the pose after the player evaluates
+	TBaseBlendedCurve<FDefaultAllocator, UE::Anim::FCurveElement> OverrideCurve;
 
+private:
 	// Embedded standalone player to play sequence
 	FAnimNode_SequencePlayer_Standalone SequencePlayerNode;
 
@@ -62,10 +69,6 @@ public:
 	// the output FPoseContext will be from StoredPose, StoredCurve, StoredAttributes
 	FCompactHeapPose StoredPose;
 	FBlendedHeapCurve StoredCurve;
-
-	// Curves to add to the pose after the player evaluates
-	TBaseBlendedCurve<FDefaultAllocator, UE::Anim::FCurveElement> OverrideCurve;
-	
 	UE::Anim::FHeapAttributeContainer StoredAttributes;
 	// We need to store the bone container, in case we have a LOD swap during a blend that uses the stored pose.
 	FBoneContainer StoredBoneContainer;
@@ -119,7 +122,7 @@ struct BLENDSTACK_API FAnimNode_BlendStack_Standalone : public FAnimNode_AssetPl
 		bool bMirrored = false, UMirrorDataTable* MirrorDataTable = nullptr,
 		float BlendTime = 0.2f, float RootBoneBlendTime = -1.f, float MaxTimeBeforeFreezingInnerBlends = -1.f,
 		const UBlendProfile* BlendProfile = nullptr, EAlphaBlendOption BlendOption = EAlphaBlendOption::Linear, 
-		bool bUseInertialBlend = false, FVector BlendParameters = FVector::Zero(), float PlayRate = 1.f,
+		bool bUseInertialBlend = false, const FVector& BlendParameters = FVector::Zero(), float PlayRate = 1.f,
 		FName GroupName = NAME_None, EAnimGroupRole::Type GroupRole = EAnimGroupRole::CanBeLeader, EAnimSyncMethod Method = EAnimSyncMethod::DoNotSync);
 	void UpdatePlayRate(float PlayRate);
 	void Reset();
@@ -142,8 +145,14 @@ protected:
 	UPROPERTY(EditAnywhere, Category = Settings, meta = (ClampMin = "0"))
 	int32 MaxActiveBlends = 4;
 
-private:
+	// if the number of requested blends is higher than MaxActiveBlends, blend stack will blend and accumulate 
+	// into a stored pose all the overflowing animations. if bStoreBlendedPose is false, the memory to store the pose will be saved,
+	// but once reached the MaxActiveBlends, blendstack will start discarding animations, potentially resulting in animation pops
+	UPROPERTY(EditAnywhere, Category = Settings)
+	bool bStoreBlendedPose = true;
 
+private:
+	void PopLastAnimPlayer();
 	void InitializeSample(const FAnimationInitializeContext& Context, FBlendStackAnimPlayer& SamplePlayer);
 	void EvaluateSample(FPoseContext& Output, const int32 PlayerIndex);
 	void UpdateSample(const FAnimationUpdateContext& Context, const int32 PlayerIndex);
diff --git a/Engine/Plugins/Experimental/Chooser/Source/Chooser/Private/AnimNode_ChooserPlayer.cpp b/Engine/Plugins/Experimental/Chooser/Source/Chooser/Private/AnimNode_ChooserPlayer.cpp
--- a/Engine/Plugins/Experimental/Chooser/Source/Chooser/Private/AnimNode_ChooserPlayer.cpp
+++ b/Engine/Plugins/Experimental/Chooser/Source/Chooser/Private/AnimNode_ChooserPlayer.cpp
@@ -161,15 +161,13 @@ void FAnimNode_ChooserPlayer::UpdateAssetPlayer(const FAnimationUpdateContext& C
 	}
 
 	// Update blend space parameters
+	const FVector BlendParameters(BlendSpaceX, BlendSpaceY, 0);
 	if (bUpdateAllActiveBlendSpaces)
 	{
 		// apply blend space parameters to all blendspaces that are playing, including ones that are blending out
-		for(FBlendStackAnimPlayer& Player : AnimPlayers)
+		for (FBlendStackAnimPlayer& Player : AnimPlayers)
 		{
-			if (Cast<UBlendSpace>(Player.GetAnimationAsset()))
-			{
-				Player.BlendSpacePlayerNode.SetPosition(FVector(BlendSpaceX,BlendSpaceY,0));
-			}
+			Player.SetBlendParameters(BlendParameters);
 		}
 	}
 	else
@@ -177,11 +175,7 @@ void FAnimNode_ChooserPlayer::UpdateAssetPlayer(const FAnimationUpdateContext& C
 		// apply blend space parameters only to the blendspace that is playing/blending in
 		if (!AnimPlayers.IsEmpty())
 		{
-			FBlendStackAnimPlayer& Player = AnimPlayers.First();
-			if (Cast<UBlendSpace>(Player.GetAnimationAsset()))
-			{
-				Player.BlendSpacePlayerNode.SetPosition(FVector(BlendSpaceX,BlendSpaceY,0));
-			}
+			AnimPlayers.First().SetBlendParameters(BlendParameters);
 		}
 	}
 
@@ -273,15 +267,7 @@ bool FAnimNode_ChooserPlayer::IsLooping() const
 {
 	if (!AnimPlayers.IsEmpty())
 	{
-		const FBlendStackAnimPlayer& AnimPlayer = AnimPlayers.First();
-		if (Cast<UBlendSpace>(AnimPlayer.GetAnimationAsset()))
-		{
-			return AnimPlayer.BlendSpacePlayerNode.IsLooping();
-		}
-		else
-		{
-			return AnimPlayer.SequencePlayerNode.IsLooping();
-		}
+		return AnimPlayers.First().IsLooping();
 	}
 	return false;
 }
