Fix for crash in animnode_blendstack when using a stored pose with an invalid bone container

diff --git a/Engine/Plugins/Experimental/Animation/BlendStack/Source/Runtime/Private/AnimNode_BlendStack.cpp b/Engine/Plugins/Experimental/Animation/BlendStack/Source/Runtime/Private/AnimNode_BlendStack.cpp
--- a/Engine/Plugins/Experimental/Animation/BlendStack/Source/Runtime/Private/AnimNode_BlendStack.cpp
+++ b/Engine/Plugins/Experimental/Animation/BlendStack/Source/Runtime/Private/AnimNode_BlendStack.cpp
@@ -177,13 +177,13 @@ void FBlendStackAnimPlayer::RestorePoseContext(FPoseContext& PoseContext) const
 {
 	check(!SequencePlayerNode.GetSequence() && !BlendSpacePlayerNode.GetBlendSpace());
 
-	if (StoredPose.IsValid())
+	if (StoredBoneContainer.IsValid())
 	{
 		// Serial number mismatch means a potential bone LOD mismatch, even if we have the same number of bones.
 		// Remap the pose manually in those cases.
 		if (PoseContext.Pose.GetBoneContainer().GetSerialNumber() == StoredBoneContainer.GetSerialNumber())
 		{
-			PoseContext.Pose.CopyBonesFrom(StoredPose);
+			PoseContext.Pose.CopyBonesFrom(StoredPose.GetBones());
 		}
 		else
 		{
