blend stack - bone container caching cleanup

diff --git a/Engine/Plugins/Experimental/Animation/BlendStack/Source/Runtime/Private/AnimNode_BlendStack.cpp b/Engine/Plugins/Experimental/Animation/BlendStack/Source/Runtime/Private/AnimNode_BlendStack.cpp
--- a/Engine/Plugins/Experimental/Animation/BlendStack/Source/Runtime/Private/AnimNode_BlendStack.cpp
+++ b/Engine/Plugins/Experimental/Animation/BlendStack/Source/Runtime/Private/AnimNode_BlendStack.cpp
@@ -152,8 +152,8 @@ void FBlendStackAnimPlayer::StorePoseContext(const FPoseContext& PoseContext)
 
 	if (PoseContext.Pose.IsValid())
 	{
-		StoredPose.CopyBonesFrom(PoseContext.Pose);
-		StoredPose.CopyAndAssignBoneContainer(StoredBoneContainer);
+		StoredBones = PoseContext.Pose.GetBones();
+		StoredBoneContainer = PoseContext.Pose.GetBoneContainer();
 	}
 
 	StoredCurve.CopyFrom(PoseContext.Curve);
@@ -162,12 +162,12 @@ void FBlendStackAnimPlayer::StorePoseContext(const FPoseContext& PoseContext)
 
 bool FBlendStackAnimPlayer::HasValidPoseContext() const
 {
-	return StoredPose.IsValid();
+	return StoredBoneContainer.IsValid();
 }
 
 void FBlendStackAnimPlayer::MovePoseContextTo(FBlendStackAnimPlayer& Other)
 {
-	Other.StoredPose = MoveTemp(StoredPose);
+	Other.StoredBones = MoveTemp(StoredBones);
 	Other.StoredCurve = MoveTemp(StoredCurve);
 	Other.StoredAttributes = MoveTemp(StoredAttributes);
 	Other.StoredBoneContainer = MoveTemp(StoredBoneContainer);
@@ -183,7 +183,8 @@ void FBlendStackAnimPlayer::RestorePoseContext(FPoseContext& PoseContext) const
 		// Remap the pose manually in those cases.
 		if (PoseContext.Pose.GetBoneContainer().GetSerialNumber() == StoredBoneContainer.GetSerialNumber())
 		{
-			PoseContext.Pose.CopyBonesFrom(StoredPose.GetBones());
+			check(PoseContext.Pose.GetNumBones() == StoredBones.Num());
+			FMemory::Memcpy(PoseContext.Pose.GetMutableBones().GetData(), StoredBones.GetData(), sizeof(FTransform) * PoseContext.Pose.GetNumBones());
 		}
 		else
 		{
@@ -200,7 +201,7 @@ void FBlendStackAnimPlayer::RestorePoseContext(FPoseContext& PoseContext) const
 				}
 				else
 				{
-					PoseContext.Pose[CompactPoseIndex] = StoredPose[StoredCompactPoseIndex];
+					PoseContext.Pose[CompactPoseIndex] = StoredBones[StoredCompactPoseIndex.GetInt()];
 				}
 			}
 		}
diff --git a/Engine/Plugins/Experimental/Animation/BlendStack/Source/Runtime/Public/BlendStack/AnimNode_BlendStack.h b/Engine/Plugins/Experimental/Animation/BlendStack/Source/Runtime/Public/BlendStack/AnimNode_BlendStack.h
--- a/Engine/Plugins/Experimental/Animation/BlendStack/Source/Runtime/Public/BlendStack/AnimNode_BlendStack.h
+++ b/Engine/Plugins/Experimental/Animation/BlendStack/Source/Runtime/Public/BlendStack/AnimNode_BlendStack.h
@@ -66,8 +66,9 @@ private:
 
 	// if SequencePlayerNode.GetSequence() and BlendSpacePlayerNode.GetBlendSpace() are nullptr, 
 	// instead of using SequencePlayerNode or BlendSpacePlayerNode (wrapped in MirrorNode),
-	// the output FPoseContext will be from StoredPose, StoredCurve, StoredAttributes
-	FCompactHeapPose StoredPose;
+	// the output FPoseContext will be from StoredBones, StoredCurve, StoredAttributes
+	// NoTe: we don't need a full FCompactHeapPose, since we use StoredBoneContainer to cache the bone container
+	TArray<FTransform> StoredBones;
 	FBlendedHeapCurve StoredCurve;
 	UE::Anim::FHeapAttributeContainer StoredAttributes;
 	// We need to store the bone container, in case we have a LOD swap during a blend that uses the stored pose.
