Make "a.AnimNode.BlendStack.Enable 0" disable inertial blends as well

diff --git a/Engine/Plugins/Experimental/Animation/BlendStack/Source/Runtime/Private/AnimNode_BlendStack.cpp b/Engine/Plugins/Experimental/Animation/BlendStack/Source/Runtime/Private/AnimNode_BlendStack.cpp
--- a/Engine/Plugins/Experimental/Animation/BlendStack/Source/Runtime/Private/AnimNode_BlendStack.cpp
+++ b/Engine/Plugins/Experimental/Animation/BlendStack/Source/Runtime/Private/AnimNode_BlendStack.cpp
@@ -730,6 +730,13 @@ float FAnimNode_BlendStack_Standalone::GetAccumulatedTime() const
 
 static void RequestInertialBlend(const FAnimationUpdateContext& Context, float BlendTime, const UBlendProfile* BlendProfile, EAlphaBlendOption BlendOption)
 {
+#if ENABLE_ANIM_DEBUG
+	if (!CVarAnimBlendStackEnable.GetValueOnAnyThread())
+	{
+		return;
+	}
+#endif // ENABLE_ANIM_DEBUG
+
 	if (BlendTime > 0.0f)
 	{
 		UE::Anim::IInertializationRequester* InertializationRequester = Context.GetMessage<UE::Anim::IInertializationRequester>();
