blend stack
- improved GC compliance
- fixed memory stomp

diff --git a/Engine/Plugins/Experimental/Animation/BlendStack/Source/Runtime/Private/AnimNode_BlendStack.cpp b/Engine/Plugins/Experimental/Animation/BlendStack/Source/Runtime/Private/AnimNode_BlendStack.cpp
--- a/Engine/Plugins/Experimental/Animation/BlendStack/Source/Runtime/Private/AnimNode_BlendStack.cpp
+++ b/Engine/Plugins/Experimental/Animation/BlendStack/Source/Runtime/Private/AnimNode_BlendStack.cpp
@@ -15,7 +15,6 @@
 
 #if ENABLE_ANIM_DEBUG
 static TAutoConsoleVariable<bool> CVarAnimBlendStackEnable(TEXT("a.AnimNode.BlendStack.Enable"), true, TEXT("Enable / Disable Blend Stack"));
-static TAutoConsoleVariable<bool> CVarAnimBlendStackPruningEnable(TEXT("a.AnimNode.BlendStack.Pruning.Enable"), true, TEXT("Enable / Disable Blend Stack Pruning"));
 #endif
 
 #define LOCTEXT_NAMESPACE "AnimNode_BlendStack"
@@ -420,7 +419,8 @@ void FAnimNode_BlendStack_Standalone::PopLastAnimPlayer()
 		AnimPlayers[LastAnimPlayerIndex].MovePoseContextTo(AnimPlayers[LastAnimPlayerIndex - 1]);
 	}
 
-	AnimPlayers.PopLast();
+	// popping the last anim player
+	AnimPlayers.SetNum(LastAnimPlayerIndex);
 }
 
 void FAnimNode_BlendStack_Standalone::Evaluate_AnyThread(FPoseContext& Output)
@@ -430,26 +430,18 @@ void FAnimNode_BlendStack_Standalone::Evaluate_AnyThread(FPoseContext& Output)
 
 	Super::Evaluate_AnyThread(Output);
 
+	bool bDisableBlendStack = false;
+#if ENABLE_ANIM_DEBUG
+	bDisableBlendStack = !CVarAnimBlendStackEnable.GetValueOnAnyThread();
+#endif // ENABLE_ANIM_DEBUG
+
 	const int32 BlendStackSize = AnimPlayers.Num();
 	if (BlendStackSize <= 0)
 	{
 		Output.ResetToRefPose();
 	}
-	else if (BlendStackSize == 1)
-	{
-		EvaluateSample(Output, 0);
-	}
-	else if (MaxActiveBlends <= 0 
-#if ENABLE_ANIM_DEBUG
-		|| !CVarAnimBlendStackEnable.GetValueOnAnyThread()
-#endif // ENABLE_ANIM_DEBUG
-		)
+	else if (BlendStackSize == 1 || bDisableBlendStack)
 	{
-		// Disable blend stack if requested (for testing / debugging) by removing all the AnimPlayers except the first
-		while (AnimPlayers.Num() > 1)
-		{
-			PopLastAnimPlayer();
-		}
 		EvaluateSample(Output, 0);
 	}
 	else
@@ -487,34 +479,24 @@ void FAnimNode_BlendStack_Standalone::Evaluate_AnyThread(FPoseContext& Output)
 			}
 		};
 
-#if ENABLE_ANIM_DEBUG
-		const bool bEnablePruning = CVarAnimBlendStackPruningEnable.GetValueOnAnyThread();
-#else
-		const bool bEnablePruning = true;
-#endif // ENABLE_ANIM_DEBUG
-
 		// ...continuing with the valuation and accumulation on the Output FPoseContext
 		// of AnimPlayer(s) from the second last to the AnimPlayer[MaxActiveBlends].
-		// Popping them if bEnablePruning
 		int32 PlayerIndex = BlendStackSize - 2;
 		// Start evaluating with our least significant players.
 		for (; PlayerIndex >= MaxActiveBlends; --PlayerIndex)
 		{
 			EvaluateAndBlendPlayerByIndex(PlayerIndex);
 
-			if (bEnablePruning)
-			{
 				// too many AnimPlayers! we don't have enough available blends to hold them all, so we accumulate the blended poses into Output / BlendedPoseContext.
 				PopLastAnimPlayer();
 			}
-		}
 
 		// At this point Output FPoseContext contains all the weighted accumulated poses of the from AnimPlayer[MaxActiveBlends] to AnimPlayer[AnimPlayer.Num()-1]
 		if (PlayerIndex == (MaxActiveBlends - 1))
 		{
 			check(AnimPlayers.Num() == MaxActiveBlends + 1);
 
-			if (bEnablePruning && bStoreBlendedPose)
+			if (bStoreBlendedPose)
 			{
 				// We store Output / BlendedPoseContext into the last AnimPlayer, that will hold a static pose, no longer an animation playing.
 				AnimPlayers.Last().StorePoseContext(Output);
@@ -537,12 +519,6 @@ void FAnimNode_BlendStack_Standalone::Evaluate_AnyThread(FPoseContext& Output)
 		{
 			EvaluateAndBlendPlayerByIndex(PlayerIndex);
 		}
-
-		const int32 ActiveBlends = AnimPlayers.Num() - 1;
-		if (ActiveBlends > MaxActiveBlends)
-		{
-			UE_LOG(LogBlendStack, Display, TEXT("FAnimNode_BlendStack_Standalone NumBlends/MaxNumBlends %d / %d"), ActiveBlends, MaxActiveBlends);
-		}
 	}
 }
 
@@ -612,7 +588,10 @@ void FAnimNode_BlendStack_Standalone::UpdateAssetPlayer(const FAnimationUpdateCo
 		CurrentWeightMultiplier *= (1.f - BlendInPercentage);
 	}
 
-	// AnimPlayers[AnimPlayerIndex] is the first FBlendStackAnimPlayer with a weight contribution of zero, so we can discard it and all the successive AnimPlayers as well
+	// AnimPlayers[AnimPlayerIndex] is the first FBlendStackAnimPlayer with a weight contribution of zero,
+	// so we can discard it and all the successive AnimPlayers as well
+	// NoTe that it's safe to delete all those players, becasue we didn't call SamplePlayer.Update_AnyThread 
+	// hence not register sequence / blendspace player InternalTimeAccumulator via FAnimTickRecord(s)
 	const int32 WantedAnimPlayersNum = FMath::Max(1, AnimPlayerIndex); // we save at least one FBlendStackAnimPlayer
 	while (AnimPlayers.Num() > WantedAnimPlayersNum)
 	{
@@ -720,12 +699,12 @@ void FAnimNode_BlendStack_Standalone::InitializeSample(const FAnimationInitializ
 
 UAnimationAsset* FAnimNode_BlendStack_Standalone::GetAnimAsset() const
 {
-	return AnimPlayers.IsEmpty() ? nullptr : AnimPlayers.First().GetAnimationAsset();
+	return AnimPlayers.IsEmpty() ? nullptr : AnimPlayers[0].GetAnimationAsset();
 }
 
 float FAnimNode_BlendStack_Standalone::GetAccumulatedTime() const
 {
-	return AnimPlayers.IsEmpty() ? 0.f : AnimPlayers.First().GetAccumulatedTime();
+	return AnimPlayers.IsEmpty() ? 0.f : AnimPlayers[0].GetAccumulatedTime();
 }
 
 static void RequestInertialBlend(const FAnimationUpdateContext& Context, float BlendTime, const UBlendProfile* BlendProfile, EAlphaBlendOption BlendOption)
@@ -773,8 +752,17 @@ void FAnimNode_BlendStack_Standalone::BlendTo(const FAnimationUpdateContext& Con
 		BlendTime = 0.0f;
 	}
 
-	AnimPlayers.PushFirst(FBlendStackAnimPlayer());
-	FBlendStackAnimPlayer& AnimPlayer = AnimPlayers.First();
+	if (AnimPlayers.Num() <= MaxActiveBlends + 2)
+	{
+		AnimPlayers.Insert(FBlendStackAnimPlayer(), 0);
+	}
+	else
+	{
+		// else it means we had multiple BlendTo suring the same frame. we'll let the last one win
+		UE_LOG(LogBlendStack, Warning, TEXT("FAnimNode_BlendStack_Standalone multiple BlendTo requests during the same frame: only the last request will be put on this BlendStack"));
+	}
+
+	FBlendStackAnimPlayer& AnimPlayer = AnimPlayers[0];
 
 	FAnimationInitializeContext InitContext(Context.AnimInstanceProxy, Context.SharedContext);
 	AnimPlayer.Initialize(InitContext, AnimationAsset, AccumulatedTime, bLoop, bMirrored, MirrorDataTable, BlendTime, RootBoneBlendTime, MaxTimeBeforeFreezingInnerBlends, BlendProfile, BlendOption, BlendParameters, PlayRate, GetNextPoseLinkIndex(), GroupName, GroupRole, GroupMethod);
@@ -783,6 +771,8 @@ void FAnimNode_BlendStack_Standalone::BlendTo(const FAnimationUpdateContext& Con
 
 void FAnimNode_BlendStack_Standalone::Reset()
 {
+	// reserving MaxActiveBlends + 2 AnimPlayers, to avoid any reallocation
+	AnimPlayers.Reserve(MaxActiveBlends + 2);
 	AnimPlayers.Reset();
 }
 
@@ -804,7 +794,7 @@ void FAnimNode_BlendStack_Standalone::UpdatePlayRate(float PlayRate)
 {
 	if (!AnimPlayers.IsEmpty())
 	{
-		AnimPlayers.First().UpdatePlayRate(PlayRate);
+		AnimPlayers[0].UpdatePlayRate(PlayRate);
 	}
 }
 
@@ -935,7 +925,7 @@ void FAnimNode_BlendStack::UpdateAssetPlayer(const FAnimationUpdateContext& Cont
 	}
 	else
 	{
-		const FBlendStackAnimPlayer& MainAnimPlayer = AnimPlayers.First();
+		const FBlendStackAnimPlayer& MainAnimPlayer = AnimPlayers[0];
 		const UAnimationAsset* PlayingAnimationAsset = MainAnimPlayer.GetAnimationAsset();
 
 		if (bForceBlendNextUpdate)
diff --git a/Engine/Plugins/Experimental/Animation/BlendStack/Source/Runtime/Public/BlendStack/AnimNode_BlendStack.h b/Engine/Plugins/Experimental/Animation/BlendStack/Source/Runtime/Public/BlendStack/AnimNode_BlendStack.h
--- a/Engine/Plugins/Experimental/Animation/BlendStack/Source/Runtime/Public/BlendStack/AnimNode_BlendStack.h
+++ b/Engine/Plugins/Experimental/Animation/BlendStack/Source/Runtime/Public/BlendStack/AnimNode_BlendStack.h
@@ -6,7 +6,6 @@
 #include "Animation/AnimNode_SequencePlayer.h"
 #include "AnimNodes/AnimNode_BlendSpacePlayer.h"
 #include "AnimNodes/AnimNode_Mirror.h"
-#include "Containers/Deque.h"
 #include "AnimNode_BlendStack.generated.h"
 
 USTRUCT()
@@ -56,12 +55,15 @@ struct BLENDSTACK_API FBlendStackAnimPlayer
 
 private:
 	// Embedded standalone player to play sequence
+	UPROPERTY(Transient)
 	FAnimNode_SequencePlayer_Standalone SequencePlayerNode;
 
 	// Embedded standalone player to play blend spaces
+	UPROPERTY(Transient)
 	FAnimNode_BlendSpacePlayer_Standalone BlendSpacePlayerNode;
 
 	// Embedded mirror node to handle mirroring
+	UPROPERTY(Transient)
 	FAnimNode_Mirror_Standalone MirrorNode;
 
 	// if SequencePlayerNode.GetSequence() and BlendSpacePlayerNode.GetBlendSpace() are nullptr, 
@@ -92,7 +94,9 @@ struct FBlendStack_SampleGraphPoseLink
 	UPROPERTY()
 	int32 RootNodeIndex = INDEX_NONE;
 
+	UPROPERTY(Transient)
 	FPoseLink Root;
+
 	FBlendStackAnimPlayer* Player = nullptr;
 	FGraphTraversalCounter CacheBoneCounter;
 
@@ -110,7 +114,8 @@ struct BLENDSTACK_API FAnimNode_BlendStack_Standalone : public FAnimNode_AssetPl
 	TArray<FBlendStack_SampleGraphPoseLink> SampleGraphPoseLinks;
 	int32 CurrentSamplePoseLink = -1;
 
-	TDeque<FBlendStackAnimPlayer> AnimPlayers;
+	UPROPERTY(Transient)
+	TArray<FBlendStackAnimPlayer> AnimPlayers;
 
 	// FAnimNode_Base interface
 	virtual void Evaluate_AnyThread(FPoseContext& Output) override;
diff --git a/Engine/Plugins/Experimental/Animation/PoseSearch/Source/Runtime/Private/AnimNode_MotionMatching.cpp b/Engine/Plugins/Experimental/Animation/PoseSearch/Source/Runtime/Private/AnimNode_MotionMatching.cpp
--- a/Engine/Plugins/Experimental/Animation/PoseSearch/Source/Runtime/Private/AnimNode_MotionMatching.cpp
+++ b/Engine/Plugins/Experimental/Animation/PoseSearch/Source/Runtime/Private/AnimNode_MotionMatching.cpp
@@ -143,7 +143,7 @@ void FAnimNode_MotionMatching::UpdateAssetPlayer(const FAnimationUpdateContext&
 	if (CVarAnimNodeMotionMatchingDrawInfo.GetValueOnAnyThread())
 	{
 		const UPoseSearchDatabase* CurrentDatabase = MotionMatchingState.CurrentSearchResult.Database.Get();
-		const UAnimationAsset* CurrentAnimationAsset = AnimPlayers.IsEmpty() ? nullptr : AnimPlayers.First().GetAnimationAsset();
+		const UAnimationAsset* CurrentAnimationAsset = AnimPlayers.IsEmpty() ? nullptr : AnimPlayers[0].GetAnimationAsset();
 
 		FString DebugInfo = FString::Printf(TEXT("bForceInterruptNextUpdate(%d)\n"), bForceInterruptNextUpdate);
 		DebugInfo += FString::Printf(TEXT("Current Database(%s)\n"), *GetNameSafe(CurrentDatabase));
diff --git a/Engine/Plugins/Experimental/Chooser/Source/Chooser/Private/AnimNode_ChooserPlayer.cpp b/Engine/Plugins/Experimental/Chooser/Source/Chooser/Private/AnimNode_ChooserPlayer.cpp
--- a/Engine/Plugins/Experimental/Chooser/Source/Chooser/Private/AnimNode_ChooserPlayer.cpp
+++ b/Engine/Plugins/Experimental/Chooser/Source/Chooser/Private/AnimNode_ChooserPlayer.cpp
@@ -97,7 +97,7 @@ void FAnimNode_ChooserPlayer::UpdateAssetPlayer(const FAnimationUpdateContext& C
 	{
 		if (!AnimPlayers.IsEmpty())
 		{
-			FBlendStackAnimPlayer& Player = AnimPlayers.First();
+			FBlendStackAnimPlayer& Player = AnimPlayers[0];
 			if (Player.GetPlayRate() * Context.GetDeltaTime() + Player.GetAccumulatedTime() > CurrentAsset->GetPlayLength())
 			{
 				NewAsset = ChooseAsset(Context);
@@ -148,7 +148,7 @@ void FAnimNode_ChooserPlayer::UpdateAssetPlayer(const FAnimationUpdateContext& C
 
 		if (!Settings.CurveOverrides.Values.IsEmpty())
 		{
-			TBaseBlendedCurve<FDefaultAllocator, UE::Anim::FCurveElement>& OverrideCurve = AnimPlayers.First().OverrideCurve;
+			TBaseBlendedCurve<FDefaultAllocator, UE::Anim::FCurveElement>& OverrideCurve = AnimPlayers[0].OverrideCurve;
 			OverrideCurve.Reserve(Settings.CurveOverrides.Values.Num());
 			for (const FAnimCurveOverride& CurveOverride : Settings.CurveOverrides.Values)
 			{
@@ -175,7 +175,7 @@ void FAnimNode_ChooserPlayer::UpdateAssetPlayer(const FAnimationUpdateContext& C
 		// apply blend space parameters only to the blendspace that is playing/blending in
 		if (!AnimPlayers.IsEmpty())
 		{
-			AnimPlayers.First().SetBlendParameters(BlendParameters);
+			AnimPlayers[0].SetBlendParameters(BlendParameters);
 		}
 	}
 
@@ -267,7 +267,7 @@ bool FAnimNode_ChooserPlayer::IsLooping() const
 {
 	if (!AnimPlayers.IsEmpty())
 	{
-		return AnimPlayers.First().IsLooping();
+		return AnimPlayers[0].IsLooping();
 	}
 	return false;
 }
