blend stack
- not running grapsh on stored poses
- fix for FBlendStackAnimPlayer::HasValidPoseContext

diff --git a/Engine/Plugins/Experimental/Animation/BlendStack/Source/Runtime/Private/AnimNode_BlendStack.cpp b/Engine/Plugins/Experimental/Animation/BlendStack/Source/Runtime/Private/AnimNode_BlendStack.cpp
--- a/Engine/Plugins/Experimental/Animation/BlendStack/Source/Runtime/Private/AnimNode_BlendStack.cpp
+++ b/Engine/Plugins/Experimental/Animation/BlendStack/Source/Runtime/Private/AnimNode_BlendStack.cpp
@@ -161,15 +161,19 @@ void FBlendStackAnimPlayer::StorePoseContext(const FPoseContext& PoseContext)
 
 bool FBlendStackAnimPlayer::HasValidPoseContext() const
 {
-	return StoredBoneContainer.IsValid();
+	return !StoredBones.IsEmpty() && StoredBoneContainer.IsValid();
 }
 
 void FBlendStackAnimPlayer::MovePoseContextTo(FBlendStackAnimPlayer& Other)
 {
+	// moving the allocated memory to Other
 	Other.StoredBones = MoveTemp(StoredBones);
 	Other.StoredCurve = MoveTemp(StoredCurve);
 	Other.StoredAttributes = MoveTemp(StoredAttributes);
 	Other.StoredBoneContainer = MoveTemp(StoredBoneContainer);
+	
+	// making sure Other pose context is invalid
+	Other.StoredBones.Reset();
 }
 
 void FBlendStackAnimPlayer::RestorePoseContext(FPoseContext& PoseContext) const
@@ -488,8 +492,8 @@ void FAnimNode_BlendStack_Standalone::Evaluate_AnyThread(FPoseContext& Output)
 			EvaluateAndBlendPlayerByIndex(PlayerIndex);
 
 				// too many AnimPlayers! we don't have enough available blends to hold them all, so we accumulate the blended poses into Output / BlendedPoseContext.
-				PopLastAnimPlayer();
-			}
+			PopLastAnimPlayer();
+		}
 
 		// At this point Output FPoseContext contains all the weighted accumulated poses of the from AnimPlayer[MaxActiveBlends] to AnimPlayer[AnimPlayer.Num()-1]
 		if (PlayerIndex == (MaxActiveBlends - 1))
@@ -501,17 +505,6 @@ void FAnimNode_BlendStack_Standalone::Evaluate_AnyThread(FPoseContext& Output)
 				// We store Output / BlendedPoseContext into the last AnimPlayer, that will hold a static pose, no longer an animation playing.
 				AnimPlayers.Last().StorePoseContext(Output);
 			}
-
-			// we execute the associated graph on the Output
-			if (!SampleGraphPoseLinks.IsEmpty())
-			{
-				const int32 PoseLinkIdx = AnimPlayers[MaxActiveBlends].GetPoseLinkIndex();
-				FBlendStack_SampleGraphPoseLink& PoseLink = SampleGraphPoseLinks[PoseLinkIdx];
-
-				// No players should have evaluated a graph before this point.
-				// Evaluate the graph on the blended result.
-				PoseLink.EvaluatePlayer(Output, AnimPlayers[MaxActiveBlends]);
-			}
 		}
 
 		// Continue with the evaluation of the most significant AnimPlayer(s) with the associated graphs
@@ -602,19 +595,15 @@ void FAnimNode_BlendStack_Standalone::UpdateAssetPlayer(const FAnimationUpdateCo
 bool FAnimNode_BlendStack_Standalone::IsSampleGraphAvailableForPlayer(const int32 PlayerIndex)
 {
 	// If we have any sample graphs, our player has been assigned a pose link index.
-	// If we are within X most relelvant players, then the graph is available.
-	return !SampleGraphPoseLinks.IsEmpty() && (PlayerIndex < MaxActiveBlends);
+	// Players with a stored pose don't need to run the graph.
+	return !SampleGraphPoseLinks.IsEmpty() && !AnimPlayers[PlayerIndex].HasValidPoseContext();
 }
 
 void FAnimNode_BlendStack_Standalone::EvaluateSample(FPoseContext& Output, const int32 PlayerIndex)
 {
 	FBlendStackAnimPlayer& SamplePlayer = AnimPlayers[PlayerIndex];
-	// If we have any sample graphs, our player has been assigned a pose link index.
-	// If we are within X most relelvant players, then the graph is available.
-	// If PlayerIndex == MaxActiveBlends, don't evaluate that graph. It's reserved for the stored pose.
 	// MaxActiveBlends == 0, means we're using inertialization. Run the the graph.
-	const bool bIsSampleGraphAvailable = !SampleGraphPoseLinks.IsEmpty() && 
-										((PlayerIndex < MaxActiveBlends) || (MaxActiveBlends == 0));
+	const bool bIsSampleGraphAvailable = IsSampleGraphAvailableForPlayer(PlayerIndex);
 	if (!bIsSampleGraphAvailable)
 	{
 		// If we have no sample graph, evaluate the player directly.
@@ -653,11 +642,7 @@ void FAnimNode_BlendStack_Standalone::UpdateSample(const FAnimationUpdateContext
 {
 	FBlendStackAnimPlayer& SamplePlayer = AnimPlayers[PlayerIndex];
 
-	// If we have any sample graphs, our player has been assigned a pose link index.
-	// If we are within X most relelvant players, then the graph is available.
-	// @todo: If PlayerIndex == MaxActiveBlends, this will likely become a stored pose. What do we update that graph with? 
-	// For now, just use the same player.
-	const bool bHasSampleGraph = !SampleGraphPoseLinks.IsEmpty() && (PlayerIndex <= MaxActiveBlends);
+	const bool bHasSampleGraph = IsSampleGraphAvailableForPlayer(PlayerIndex);
 	if (bHasSampleGraph)
 	{
 		FBlendStack_SampleGraphPoseLink& PoseLink = SampleGraphPoseLinks[SamplePlayer.GetPoseLinkIndex()];
@@ -678,7 +663,7 @@ void FAnimNode_BlendStack_Standalone::UpdateSample(const FAnimationUpdateContext
 void FAnimNode_BlendStack_Standalone::CacheBonesForSample(const FAnimationCacheBonesContext& Context, const int32 PlayerIndex)
 {
 	FBlendStackAnimPlayer& SamplePlayer = AnimPlayers[PlayerIndex];
-	const bool bHasSampleGraph = !SampleGraphPoseLinks.IsEmpty() && (PlayerIndex <= MaxActiveBlends);
+	const bool bHasSampleGraph = IsSampleGraphAvailableForPlayer(PlayerIndex);
 	if (bHasSampleGraph)
 	{
 		FBlendStack_SampleGraphPoseLink& PoseLink = SampleGraphPoseLinks[SamplePlayer.GetPoseLinkIndex()];
