First pass for notify filtering support for blend stack.

- Added IAnimEventsFilterContextData. Interface context data class used by scope graph messages to request notifies to be filtered.
- Added FBlendStackAnimEventsFilterScope. This keeps track of notifies fired from a blend stack and determines if a notify originating from a sample can be queued.
- Updated AnimNotifyQueue to support filter notifies if requested.

diff --git a/Engine/Plugins/Experimental/Animation/BlendStack/Source/Runtime/Private/AnimNode_BlendStack.cpp b/Engine/Plugins/Experimental/Animation/BlendStack/Source/Runtime/Private/AnimNode_BlendStack.cpp
--- a/Engine/Plugins/Experimental/Animation/BlendStack/Source/Runtime/Private/AnimNode_BlendStack.cpp
+++ b/Engine/Plugins/Experimental/Animation/BlendStack/Source/Runtime/Private/AnimNode_BlendStack.cpp
@@ -9,6 +9,7 @@
 #include "Animation/AnimMontage.h"
 #include "BlendStack/AnimNode_BlendStackInput.h"
 #include "Animation/AnimNode_Inertialization.h"
+#include "BlendStackAnimEventsFilterScope.h"
 #include "BlendStack/BlendStackDefines.h"
 
 #include UE_INLINE_GENERATED_CPP_BY_NAME(AnimNode_BlendStack)
@@ -519,6 +520,12 @@ void FAnimNode_BlendStack_Standalone::Initialize_AnyThread(const FAnimationIniti
 {
 	Super::Initialize_AnyThread(Context);
 
+	if (bShouldFilterNotifies)
+	{
+		NotifiesFiredLastTick = MakeShared<TArray<FName>>();
+		NotifyRecencyMap = MakeShared<TMap<FName, float>>();
+	}
+	
 	Reset();
 
 	if (SampleGraphPoseLinks.IsEmpty() == false)
@@ -557,6 +564,32 @@ void FAnimNode_BlendStack_Standalone::UpdateAssetPlayer(const FAnimationUpdateCo
 
 	Super::UpdateAssetPlayer(Context);
 
+	if (bShouldFilterNotifies)
+	{
+		if (const UWorld* World = Context.GetAnimInstanceObject()->GetWorld())
+		{
+			const double CurrentGameTime = World->GetTimeSeconds();
+		
+			// Set target time-outs for notifies fired last tick.
+			for (const FName & NotifyName : *NotifiesFiredLastTick)
+			{
+				NotifyRecencyMap->FindOrAdd(NotifyName) = CurrentGameTime + NotifyRecencyTimeOut;
+			}
+			NotifiesFiredLastTick->Reset();
+			
+			// Find notifies that have timed-out and should be allowed to fire this tick.
+			for (auto It = NotifyRecencyMap->CreateIterator(); It; ++It)
+			{
+				if (CurrentGameTime >= It->Value)
+				{
+					It.RemoveCurrent();
+				}
+			}
+		}
+	}
+	
+	UE::Anim::TOptionalScopedGraphMessage<UE::Anim::FBlendStackAnimEventsFilterScope> BlendStackInfoScope(bShouldFilterNotifies, Context, NotifiesFiredLastTick, NotifyRecencyMap);
+
 	// AnimPlayers[0] is the most newly inserted AnimPlayer, AnimPlayers[AnimPlayers.Num()-1] is the oldest, so to calculate the weights
 	// we ask AnimPlayers[0] its BlendInPercentage and then distribute the left over (CurrentWeightMultiplier) to the rest of the AnimPlayers
 	// AnimPlayers[AnimPlayerIndex].GetBlendWeight() will now store the weighted contribution of AnimPlayers[AnimPlayerIndex] to be able to calculate root motion from animation
@@ -759,6 +792,12 @@ void FAnimNode_BlendStack_Standalone::Reset()
 	// reserving MaxActiveBlends + 2 AnimPlayers, to avoid any reallocation
 	AnimPlayers.Reserve(MaxActiveBlends + 2);
 	AnimPlayers.Reset();
+
+	if (bShouldFilterNotifies)
+	{
+		NotifiesFiredLastTick->Reset();
+		NotifyRecencyMap->Reset();
+	}
 }
 
 int32 FAnimNode_BlendStack_Standalone::GetNextPoseLinkIndex()
diff --git a/Engine/Plugins/Experimental/Animation/BlendStack/Source/Runtime/Private/BlendStackAnimEventsFilterScope.cpp b/Engine/Plugins/Experimental/Animation/BlendStack/Source/Runtime/Private/BlendStackAnimEventsFilterScope.cpp
new file mode 100644
--- /dev/null
+++ b/Engine/Plugins/Experimental/Animation/BlendStack/Source/Runtime/Private/BlendStackAnimEventsFilterScope.cpp
@@ -0,0 +1,44 @@
+// Copyright Epic Games, Inc. All Rights Reserved.
+
+#include "BlendStackAnimEventsFilterScope.h"
+
+IMPLEMENT_ANIMGRAPH_MESSAGE(UE::Anim::FBlendStackAnimEventsFilterScope)
+
+namespace UE { namespace Anim {
+	
+	FBlendStackAnimEventsFilterContext::FBlendStackAnimEventsFilterContext(const TSharedRef<TArray<FName>> & InFiredNotifies, const TSharedRef<TMap<FName, float>>& InNotifyBanList)
+		: FiredNotifies(InFiredNotifies), NotifyBanList(InNotifyBanList)
+	{
+	}
+
+	bool FBlendStackAnimEventsFilterContext::ShouldFilterNotify(const FAnimNotifyEventReference& InNotifyEventRef) const
+	{
+		const FName NotifyName = InNotifyEventRef.GetNotify()->NotifyName;
+		
+		const bool bWasRecentlyFired = NotifyBanList->Contains(NotifyName);
+		const bool bWasAlreadyFiredThisTick = FiredNotifies->Contains(NotifyName);
+		
+		if (!bWasAlreadyFiredThisTick && !bWasRecentlyFired)
+		{
+			FiredNotifies->Push(NotifyName);
+			return false;
+		}
+
+		return true;
+	}
+
+	FBlendStackAnimEventsFilterScope::FBlendStackAnimEventsFilterScope(const TSharedRef<TArray<FName>> & InFiredNotifies, const TSharedRef<TMap<FName, float>> & InNotifyBanList)
+		: FiredNotifies(InFiredNotifies), NotifyBanList(InNotifyBanList)
+	{
+	}
+
+	FBlendStackAnimEventsFilterScope::FBlendStackAnimEventsFilterScope(const TSharedPtr<TArray<FName>> & InFiredNotifies, const TSharedPtr<TMap<FName,float>> & InNotifyBanList)
+		: FiredNotifies(InFiredNotifies.ToSharedRef()), NotifyBanList(InNotifyBanList.ToSharedRef())
+	{
+	}
+	
+	TUniquePtr<const IAnimNotifyEventContextDataInterface> FBlendStackAnimEventsFilterScope::MakeUniqueEventContextData() const
+	{
+		return MakeUnique<const FBlendStackAnimEventsFilterContext>(FiredNotifies, NotifyBanList);
+	}
+}}	// namespace UE::Anim
\ No newline at end of file
diff --git a/Engine/Plugins/Experimental/Animation/BlendStack/Source/Runtime/Private/BlendStackAnimEventsFilterScope.h b/Engine/Plugins/Experimental/Animation/BlendStack/Source/Runtime/Private/BlendStackAnimEventsFilterScope.h
new file mode 100644
--- /dev/null
+++ b/Engine/Plugins/Experimental/Animation/BlendStack/Source/Runtime/Private/BlendStackAnimEventsFilterScope.h
@@ -0,0 +1,41 @@
+// Copyright Epic Games, Inc. All Rights Reserved.
+
+#pragma once
+
+#include "CoreMinimal.h"
+#include "Animation/AnimEventsFilterScope.h"
+#include "Animation/AnimNodeMessages.h"
+#include "Animation/AnimationAsset.h"
+
+namespace UE { namespace Anim {
+	
+	class FBlendStackAnimEventsFilterContext : public IAnimEventsFilterContext
+	{
+	public:
+
+		FBlendStackAnimEventsFilterContext() {}
+		FBlendStackAnimEventsFilterContext(const TSharedRef<TArray<FName>> & InFiredNotifies, const TSharedRef<TMap<FName,float>>& InNotifyBanList);
+		
+		virtual bool ShouldFilterNotify(const FAnimNotifyEventReference& InNotifyEventRef) const override;
+	
+		TSharedRef<TArray<FName>> FiredNotifies;
+		TSharedRef<TMap<FName,float>> NotifyBanList;
+	};
+	
+	class FBlendStackAnimEventsFilterScope : public IGraphMessage
+	{
+		DECLARE_ANIMGRAPH_MESSAGE(FBlendStackAnimEventsFilterScope);
+
+	public:
+		
+		FBlendStackAnimEventsFilterScope(const TSharedRef<TArray<FName>> & InFiredNotifies, const TSharedRef<TMap<FName,float>> & InNotifyBanList);
+		FBlendStackAnimEventsFilterScope(const TSharedPtr<TArray<FName>> & InFiredNotifies, const TSharedPtr<TMap<FName,float>> & InNotifyBanList);
+		
+		virtual TUniquePtr<const IAnimNotifyEventContextDataInterface> MakeUniqueEventContextData() const override;
+
+	private:
+		TSharedRef<TArray<FName>> FiredNotifies;
+		TSharedRef<TMap<FName,float>> NotifyBanList;
+	};
+	
+}}	// namespace UE::Anim
diff --git a/Engine/Plugins/Experimental/Animation/BlendStack/Source/Runtime/Public/BlendStack/AnimNode_BlendStack.h b/Engine/Plugins/Experimental/Animation/BlendStack/Source/Runtime/Public/BlendStack/AnimNode_BlendStack.h
--- a/Engine/Plugins/Experimental/Animation/BlendStack/Source/Runtime/Public/BlendStack/AnimNode_BlendStack.h
+++ b/Engine/Plugins/Experimental/Animation/BlendStack/Source/Runtime/Public/BlendStack/AnimNode_BlendStack.h
@@ -117,6 +117,10 @@ struct BLENDSTACK_API FAnimNode_BlendStack_Standalone : public FAnimNode_AssetPl
 	UPROPERTY(Transient)
 	TArray<FBlendStackAnimPlayer> AnimPlayers;
 
+	// Flag that determines if any notifies from originating from an anim player samples should be filtered or not.
+	UPROPERTY(EditAnywhere, Category = Settings)
+	bool bShouldFilterNotifies = false;
+	
 	// FAnimNode_Base interface
 	virtual void Evaluate_AnyThread(FPoseContext& Output) override;
 	virtual void Initialize_AnyThread(const FAnimationInitializeContext& Context) override;
@@ -157,6 +161,13 @@ protected:
 	UPROPERTY(EditAnywhere, Category = Settings)
 	bool bStoreBlendedPose = true;
 
+    TSharedPtr<TArray<FName>> NotifiesFiredLastTick;
+	TSharedPtr<TMap<FName,float>> NotifyRecencyMap;
+
+	// Window of time after firing a notify that any instance of the same notify will be filtered out.
+	UPROPERTY(EditAnywhere, BlueprintReadWrite, Category = Settings, meta = (ClampMin = "0", PinHiddenByDefault))
+	float NotifyRecencyTimeOut = 0.2f;
+
 private:
 	void PopLastAnimPlayer();
 	void InitializeSample(const FAnimationInitializeContext& Context, FBlendStackAnimPlayer& SamplePlayer);
diff --git a/Engine/Source/Runtime/Engine/Private/Animation/AnimEventsFilterScope.cpp b/Engine/Source/Runtime/Engine/Private/Animation/AnimEventsFilterScope.cpp
new file mode 100644
--- /dev/null
+++ b/Engine/Source/Runtime/Engine/Private/Animation/AnimEventsFilterScope.cpp
@@ -0,0 +1,5 @@
+// Copyright Epic Games, Inc. All Rights Reserved.
+
+#include "Animation/AnimEventsFilterScope.h"
+
+IMPLEMENT_NOTIFY_CONTEXT_INTERFACE(UE::Anim::IAnimEventsFilterContext)
\ No newline at end of file
diff --git a/Engine/Source/Runtime/Engine/Private/Animation/AnimNotifyQueue.cpp b/Engine/Source/Runtime/Engine/Private/Animation/AnimNotifyQueue.cpp
--- a/Engine/Source/Runtime/Engine/Private/Animation/AnimNotifyQueue.cpp
+++ b/Engine/Source/Runtime/Engine/Private/Animation/AnimNotifyQueue.cpp
@@ -1,6 +1,8 @@
 // Copyright Epic Games, Inc. All Rights Reserved.
 
 #include "Animation/AnimNotifyQueue.h"
+
+#include "Animation/AnimEventsFilterScope.h"
 #include "Animation/AnimInstanceProxy.h"
 #include "Components/SkeletalMeshComponent.h"
 
@@ -58,9 +60,15 @@ void FAnimNotifyQueue::AddAnimNotifiesToDest(bool bSrcIsLeader, const TArray<FAn
 			const bool bPassesDedicatedServerCheck = Notify->bTriggerOnDedicatedServer || !IsRunningDedicatedServer();
 			if (bPassesDedicatedServerCheck && Notify->TriggerWeightThreshold <= InstanceWeight && PassesFiltering(Notify) && PassesChanceOfTriggering(Notify))
 			{
-				// Only add unique AnimNotifyState instances just once. We can get multiple triggers if looping over an animation.
-				// It is the same state, so just report it once.
-				Notify->NotifyStateClass ? DestArray.AddUnique(NotifyRef) : DestArray.Add(NotifyRef);
+				const UE::Anim::IAnimEventsFilterContext* ContextData = Notify->bCanBeFilteredViaRequest ? NotifyRef.GetContextData<UE::Anim::IAnimEventsFilterContext>() : nullptr;
+				const bool bPassedScopeFilter = ContextData ? !ContextData->ShouldFilterNotify(NotifyRef) : true;
+
+				if (bPassedScopeFilter)
+				{
+					// Only add unique AnimNotifyState instances just once. We can get multiple triggers if looping over an animation.
+					// It is the same state, so just report it once.
+					Notify->NotifyStateClass ? DestArray.AddUnique(NotifyRef) : DestArray.Add(NotifyRef);	
+				}
 			}
 		}
 	}
diff --git a/Engine/Source/Runtime/Engine/Public/Animation/AnimEventsFilterScope.h b/Engine/Source/Runtime/Engine/Public/Animation/AnimEventsFilterScope.h
new file mode 100644
--- /dev/null
+++ b/Engine/Source/Runtime/Engine/Public/Animation/AnimEventsFilterScope.h
@@ -0,0 +1,21 @@
+// Copyright Epic Games, Inc. All Rights Reserved.
+
+#pragma once
+
+#include "CoreMinimal.h"
+#include "AnimNodeMessages.h"
+#include "AnimNotifyQueue.h"
+
+namespace UE { namespace Anim
+{
+
+class IAnimEventsFilterContext : public IAnimNotifyEventContextDataInterface
+{
+	DECLARE_NOTIFY_CONTEXT_INTERFACE_API(IAnimEventsFilterContext, ENGINE_API)
+
+public:
+	
+	virtual bool ShouldFilterNotify(const FAnimNotifyEventReference & InNotifyEventRef) const = 0;
+};
+	
+}}
\ No newline at end of file
diff --git a/Engine/Source/Runtime/Engine/Public/Animation/AnimTypes.h b/Engine/Source/Runtime/Engine/Public/Animation/AnimTypes.h
--- a/Engine/Source/Runtime/Engine/Public/Animation/AnimTypes.h
+++ b/Engine/Source/Runtime/Engine/Public/Animation/AnimTypes.h
@@ -341,6 +341,10 @@ struct FAnimNotifyEvent : public FAnimLinkableElement
 	UPROPERTY(EditAnywhere, BlueprintReadWrite, Category = AnimNotifyTriggerSettings, meta = (ClampMin = "0"))
 	int32 NotifyFilterLOD;
 
+	/** Allow notify event to be filtered if requested at runtime (e. g. via an Anim Graph Message) */
+	UPROPERTY(EditAnywhere, BlueprintReadWrite, Category = AnimNotifyTriggerSettings)
+	bool bCanBeFilteredViaRequest;
+	
 	/** If disabled this notify will be skipped on dedicated servers */
 	UPROPERTY(EditAnywhere, BlueprintReadWrite, Category = AnimNotifyTriggerSettings)
 	bool bTriggerOnDedicatedServer;
@@ -387,6 +391,7 @@ public:
 		, NotifyTriggerChance(1.f)
 		, NotifyFilterType(ENotifyFilterType::NoFiltering)
 		, NotifyFilterLOD(0)
+		, bCanBeFilteredViaRequest(true)
 		, bTriggerOnDedicatedServer(true)
 		, bTriggerOnFollower(false)
 #if WITH_EDITORONLY_DATA
