blend stack
- adding VLOG info
- added MaxBlendInTimeToOverrideAnimation
- added PlayerDepthBlendInTimeMultiplier
- removed MaxTimeBeforeFreezingInnerBlends

diff --git a/Engine/Plugins/Experimental/Animation/BlendStack/Source/Runtime/Private/AnimNode_BlendStack.cpp b/Engine/Plugins/Experimental/Animation/BlendStack/Source/Runtime/Private/AnimNode_BlendStack.cpp
--- a/Engine/Plugins/Experimental/Animation/BlendStack/Source/Runtime/Private/AnimNode_BlendStack.cpp
+++ b/Engine/Plugins/Experimental/Animation/BlendStack/Source/Runtime/Private/AnimNode_BlendStack.cpp
@@ -2,15 +2,16 @@
 
 #include "BlendStack/AnimNode_BlendStack.h"
 #include "Algo/MaxElement.h"
-#include "Animation/AnimInstanceProxy.h"
 #include "Animation/AnimComposite.h"
+#include "Animation/AnimInstanceProxy.h"
+#include "Animation/AnimMontage.h"
+#include "Animation/AnimNode_Inertialization.h"
 #include "Animation/AnimSequence.h"
 #include "Animation/BlendSpace.h"
-#include "Animation/AnimMontage.h"
 #include "BlendStack/AnimNode_BlendStackInput.h"
-#include "Animation/AnimNode_Inertialization.h"
 #include "BlendStackAnimEventsFilterScope.h"
 #include "BlendStack/BlendStackDefines.h"
+#include "VisualLogger/VisualLogger.h"
 
 #include UE_INLINE_GENERATED_CPP_BY_NAME(AnimNode_BlendStack)
 
@@ -25,7 +26,7 @@ static constexpr FBoneIndexType RootBoneIndexType = 0;
 /////////////////////////////////////////////////////
 // FBlendStackAnimPlayer
 void FBlendStackAnimPlayer::Initialize(const FAnimationInitializeContext& Context, UAnimationAsset* AnimationAsset, float AccumulatedTime, bool bLoop,
-	bool bMirrored, UMirrorDataTable* MirrorDataTable, float BlendTime, float RootBoneBlendTime, float MaxTimeBeforeFreezingInnerBlends, 
+	bool bMirrored, UMirrorDataTable* MirrorDataTable, float BlendTime, float RootBoneBlendTime, 
 	const UBlendProfile* BlendProfile, EAlphaBlendOption InBlendOption, const FVector& BlendParameters, float PlayRate, int32 InPoseLinkIdx,
 	FName GroupName, EAnimGroupRole::Type GroupRole, EAnimSyncMethod GroupMethod)
 {
@@ -79,9 +80,7 @@ void FBlendStackAnimPlayer::Initialize(const FAnimationInitializeContext& Contex
 	BlendOption = InBlendOption;
 
 	TotalBlendInTime = BlendTime;
-	MaxBlendInTimeAsSecondaryPlayer = MaxTimeBeforeFreezingInnerBlends >= 0 ? MaxTimeBeforeFreezingInnerBlends : BlendTime;
-	CurrentBlendInTimeAsMainPlayer = 0.f;
-	CurrentBlendInTimeAsSecondaryPlayer = 0.f;
+	CurrentBlendInTime = 0.f;
 
 	MirrorNode.SetMirrorDataTable(MirrorDataTable);
 	MirrorNode.SetMirror(bMirrored);
@@ -309,15 +308,18 @@ float FBlendStackAnimPlayer::GetPlayRate() const
 	return 0.f;
 }
 
-void FBlendStackAnimPlayer::AdvanceBlendInTime(const float DeltaTime, bool bIsMainPlayer)
+void FBlendStackAnimPlayer::AdvanceBlendInTime(const float DeltaTime, int32 PlayerDepth, float PlayerDepthBlendInTimeMultiplier)
 {
+	const bool bIsMainPlayer = PlayerDepth == 0;
+
 	if (bIsMainPlayer)
 	{
-		CurrentBlendInTimeAsMainPlayer += DeltaTime;
+		CurrentBlendInTime += DeltaTime;
 	}
 	else
 	{
-		CurrentBlendInTimeAsSecondaryPlayer = FMath::Min(CurrentBlendInTimeAsSecondaryPlayer + DeltaTime, MaxBlendInTimeAsSecondaryPlayer);
+		const float ScaledDeltaTime = DeltaTime * FMath::Pow(PlayerDepthBlendInTimeMultiplier, PlayerDepth + 1);
+		CurrentBlendInTime += ScaledDeltaTime;
 	}
 }
 
@@ -438,8 +440,26 @@ void FAnimNode_BlendStack_Standalone::Evaluate_AnyThread(FPoseContext& Output)
 	bool bDisableBlendStack = false;
 #if ENABLE_ANIM_DEBUG
 	bDisableBlendStack = !CVarAnimBlendStackEnable.GetValueOnAnyThread();
+
+	const FVector& ActorLocation = Output.AnimInstanceProxy->GetActorTransform().GetTranslation();
+	
+	TStringBuilder<1024> MessageBuilder;
+
+	for (int32 i = 0; i < AnimPlayers.Num(); ++i)
+	{
+		const FBlendStackAnimPlayer& AnimPlayer = AnimPlayers[i];
+		MessageBuilder.Appendf(TEXT("%d) t:%.2f/%.2f m:%d %s\n"), 
+			i, AnimPlayer.GetCurrentBlendInTime(), AnimPlayer.GetTotalBlendInTime(),
+			AnimPlayer.GetMirror() ? 1 : 0, *AnimPlayer.GetAnimationName());
+	}
+
+	FString VLogMessage(MessageBuilder.ToString());
+	FString VLogCategoryName;
+	VLogCategoryName.Appendf(TEXT("LogBlendStack_%d"), GetNodeIndex());
+	UE_VLOG_LOCATION(Output.GetAnimInstanceObject(), *VLogCategoryName, Verbose, ActorLocation + FVector(0, 0, 100.0f), 0.f, FColor::Yellow, TEXT("%s"), *VLogMessage);
 #endif // ENABLE_ANIM_DEBUG
 
+
 	const int32 BlendStackSize = AnimPlayers.Num();
 	if (BlendStackSize <= 0)
 	{
@@ -690,7 +710,7 @@ void FAnimNode_BlendStack_Standalone::UpdateSample(const FAnimationUpdateContext
 	}
 
 	// Advance the blend-in time regardless of whether or not the player was updated.
-	SamplePlayer.AdvanceBlendInTime(Context.GetDeltaTime(), PlayerIndex == 0);
+	SamplePlayer.AdvanceBlendInTime(Context.GetDeltaTime(), PlayerIndex, PlayerDepthBlendInTimeMultiplier);
 }
 
 void FAnimNode_BlendStack_Standalone::CacheBonesForSample(const FAnimationCacheBonesContext& Context, const int32 PlayerIndex)
@@ -759,9 +779,8 @@ static void RequestInertialBlend(const FAnimationUpdateContext& Context, float B
 }
 
 void FAnimNode_BlendStack_Standalone::BlendTo(const FAnimationUpdateContext& Context, UAnimationAsset* AnimationAsset, float AccumulatedTime, bool bLoop,
-	bool bMirrored, UMirrorDataTable* MirrorDataTable,
-	float BlendTime, float RootBoneBlendTime, float MaxTimeBeforeFreezingInnerBlends, const UBlendProfile* BlendProfile, EAlphaBlendOption BlendOption,
-	bool bUseInertialBlend, const FVector& BlendParameters, float PlayRate,
+	bool bMirrored, UMirrorDataTable* MirrorDataTable, float BlendTime, float RootBoneBlendTime, 
+	const UBlendProfile* BlendProfile, EAlphaBlendOption BlendOption, bool bUseInertialBlend, const FVector& BlendParameters, float PlayRate,
 	FName GroupName, EAnimGroupRole::Type GroupRole, EAnimSyncMethod GroupMethod)
 {
 	if (bUseInertialBlend)
@@ -770,7 +789,12 @@ void FAnimNode_BlendStack_Standalone::BlendTo(const FAnimationUpdateContext& Con
 		BlendTime = 0.0f;
 	}
 
-	if (AnimPlayers.Num() <= MaxActiveBlends + 2)
+	if (!AnimPlayers.IsEmpty() && AnimPlayers[0].GetCurrentBlendInTime() < MaxBlendInTimeToOverrideAnimation)
+	{
+		// replacing AnimPlayers[0] with this new BlendTo request
+		UE_LOG(LogBlendStack, Display, TEXT("FAnimNode_BlendStack_Standalone '%s' replaced by '%s' because blend time in is less than MaxBlendInTimeToOverrideAnimation (%.2f / %.2f)"), *AnimPlayers[0].GetAnimationName(), *GetNameSafe(AnimationAsset), AnimPlayers[0].GetCurrentBlendInTime(), MaxBlendInTimeToOverrideAnimation);
+	}
+	else if (AnimPlayers.Num() <= MaxActiveBlends + 2)
 	{
 		AnimPlayers.Insert(FBlendStackAnimPlayer(), 0);
 	}
@@ -783,7 +807,7 @@ void FAnimNode_BlendStack_Standalone::BlendTo(const FAnimationUpdateContext& Con
 	FBlendStackAnimPlayer& AnimPlayer = AnimPlayers[0];
 
 	FAnimationInitializeContext InitContext(Context.AnimInstanceProxy, Context.SharedContext);
-	AnimPlayer.Initialize(InitContext, AnimationAsset, AccumulatedTime, bLoop, bMirrored, MirrorDataTable, BlendTime, RootBoneBlendTime, MaxTimeBeforeFreezingInnerBlends, BlendProfile, BlendOption, BlendParameters, PlayRate, GetNextPoseLinkIndex(), GroupName, GroupRole, GroupMethod);
+	AnimPlayer.Initialize(InitContext, AnimationAsset, AccumulatedTime, bLoop, bMirrored, MirrorDataTable, BlendTime, RootBoneBlendTime, BlendProfile, BlendOption, BlendParameters, PlayRate, GetNextPoseLinkIndex(), GroupName, GroupRole, GroupMethod);
 	InitializeSample(InitContext, AnimPlayer);
 }
 
@@ -977,9 +1001,8 @@ void FAnimNode_BlendStack::UpdateAssetPlayer(const FAnimationUpdateContext& Cont
 
 	if (bExecuteBlendTo)
 	{
-		BlendTo(Context, AnimationAsset, AnimationTime, bLoop, bMirrored, MirrorDataTable.Get(),
-			BlendTime, RootBoneBlendTime, MaxTimeBeforeFreezingInnerBlends, 
-			BlendProfile, BlendOption, bUseInertialBlend, BlendParameters, WantedPlayRate);
+		BlendTo(Context, AnimationAsset, AnimationTime, bLoop, bMirrored, MirrorDataTable.Get(), BlendTime,
+			RootBoneBlendTime, BlendProfile, BlendOption, bUseInertialBlend, BlendParameters, WantedPlayRate);
 	}
 	
 	UpdatePlayRate(WantedPlayRate);
diff --git a/Engine/Plugins/Experimental/Animation/BlendStack/Source/Runtime/Private/BlendStackAnimNodeLibrary.cpp b/Engine/Plugins/Experimental/Animation/BlendStack/Source/Runtime/Private/BlendStackAnimNodeLibrary.cpp
--- a/Engine/Plugins/Experimental/Animation/BlendStack/Source/Runtime/Private/BlendStackAnimNodeLibrary.cpp
+++ b/Engine/Plugins/Experimental/Animation/BlendStack/Source/Runtime/Private/BlendStackAnimNodeLibrary.cpp
@@ -38,7 +38,6 @@ void UBlendStackAnimNodeLibrary::BlendTo(const FAnimUpdateContext& Context,
 				BlendStackNodePtr->MirrorDataTable,
 				BlendTime,
 				BlendStackNodePtr->RootBoneBlendTime,
-				BlendStackNodePtr->MaxTimeBeforeFreezingInnerBlends,
 				BlendStackNodePtr->BlendProfile,
 				BlendStackNodePtr->BlendOption,
 				BlendStackNodePtr->bUseInertialBlend,
diff --git a/Engine/Plugins/Experimental/Animation/BlendStack/Source/Runtime/Public/BlendStack/AnimNode_BlendStack.h b/Engine/Plugins/Experimental/Animation/BlendStack/Source/Runtime/Public/BlendStack/AnimNode_BlendStack.h
--- a/Engine/Plugins/Experimental/Animation/BlendStack/Source/Runtime/Public/BlendStack/AnimNode_BlendStack.h
+++ b/Engine/Plugins/Experimental/Animation/BlendStack/Source/Runtime/Public/BlendStack/AnimNode_BlendStack.h
@@ -14,7 +14,7 @@ struct BLENDSTACK_API FBlendStackAnimPlayer
 	GENERATED_BODY()
 	
 	void Initialize(const FAnimationInitializeContext& Context, UAnimationAsset* AnimationAsset, float AccumulatedTime, bool bLoop,
-		bool bMirrored, UMirrorDataTable* MirrorDataTable, float BlendTime, float RootBoneBlendTime, float MaxTimeBeforeFreezingInnerBlends,
+		bool bMirrored, UMirrorDataTable* MirrorDataTable, float BlendTime, float RootBoneBlendTime,
 		const UBlendProfile* BlendProfile, EAlphaBlendOption InBlendOption, const FVector& BlendParameters, float PlayRate, int32 InPoseLinkIdx,
 		FName GroupName, EAnimGroupRole::Type GroupRole, EAnimSyncMethod Method);
 	
@@ -35,8 +35,8 @@ struct BLENDSTACK_API FBlendStackAnimPlayer
 	void MovePoseContextTo(FBlendStackAnimPlayer& Other);
 
 	float GetTotalBlendInTime() const { return TotalBlendInTime; }
-	float GetCurrentBlendInTime() const { return CurrentBlendInTimeAsMainPlayer + CurrentBlendInTimeAsSecondaryPlayer; }
-	void AdvanceBlendInTime(const float DeltaTime, bool bIsMainPlayer);
+	float GetCurrentBlendInTime() const { return CurrentBlendInTime; }
+	void AdvanceBlendInTime(const float DeltaTime, int32 PlayerDepth, float PlayerDepthBlendInTimeMultiplier);
 	bool GetMirror() const { return MirrorNode.GetMirror(); }
 	FVector GetBlendParameters() const;
 	void SetBlendParameters(const FVector& BlendParameters);
@@ -81,9 +81,7 @@ private:
 	TCustomBoneIndexArray<float, FSkeletonPoseBoneIndex> TotalBlendInTimePerBone;
 
 	float TotalBlendInTime = 0.f;
-	float MaxBlendInTimeAsSecondaryPlayer = UE_BIG_NUMBER;
-	float CurrentBlendInTimeAsMainPlayer = 0.f;
-	float CurrentBlendInTimeAsSecondaryPlayer = 0.f;
+	float CurrentBlendInTime = 0.f;
 };
 
 USTRUCT(BlueprintInternalUseOnly)
@@ -129,8 +127,7 @@ struct BLENDSTACK_API FAnimNode_BlendStack_Standalone : public FAnimNode_AssetPl
 	// End of FAnimNode_Base interface
 
 	void BlendTo(const FAnimationUpdateContext& Context, UAnimationAsset* AnimationAsset, float AccumulatedTime = 0.f, bool bLoop = false, 
-		bool bMirrored = false, UMirrorDataTable* MirrorDataTable = nullptr,
-		float BlendTime = 0.2f, float RootBoneBlendTime = -1.f, float MaxTimeBeforeFreezingInnerBlends = -1.f,
+		bool bMirrored = false, UMirrorDataTable* MirrorDataTable = nullptr, float BlendTime = 0.2f, float RootBoneBlendTime = -1.f,
 		const UBlendProfile* BlendProfile = nullptr, EAlphaBlendOption BlendOption = EAlphaBlendOption::Linear, 
 		bool bUseInertialBlend = false, const FVector& BlendParameters = FVector::Zero(), float PlayRate = 1.f,
 		FName GroupName = NAME_None, EAnimGroupRole::Type GroupRole = EAnimGroupRole::CanBeLeader, EAnimSyncMethod Method = EAnimSyncMethod::DoNotSync);
@@ -168,6 +165,15 @@ protected:
 	UPROPERTY(EditAnywhere, BlueprintReadWrite, Category = Settings, meta = (ClampMin = "0", PinHiddenByDefault))
 	float NotifyRecencyTimeOut = 0.2f;
 
+	// if the most relevant (recently added) animation is within MaxBlendInTimeToOverrideAnimation, the new requested blend will take its spot
+	// otherwise a new blended will be added to the stack
+	UPROPERTY(EditAnywhere, Category = Settings)
+	float MaxBlendInTimeToOverrideAnimation = 0.f;
+
+	// AnimPlayers blend in timer will be incremented PlayerDepthBlendInTimeMultiplier times faster on a deeper blend
+	UPROPERTY(EditAnywhere, Category = Settings)
+	float PlayerDepthBlendInTimeMultiplier = 1.f;
+
 private:
 	void PopLastAnimPlayer();
 	void InitializeSample(const FAnimationInitializeContext& Context, FBlendStackAnimPlayer& SamplePlayer);
@@ -213,10 +219,6 @@ struct BLENDSTACK_API FAnimNode_BlendStack : public FAnimNode_BlendStack_Standal
 	UPROPERTY(EditAnywhere, BlueprintReadWrite, Category = Settings, meta = (PinHiddenByDefault))
 	float RootBoneBlendTime = -1.f;
 
-	// Max time in seconds before freezing the blend ratio of all blends, but the most recently requested one. Negative values would imply MaxTimeBeforeFreezingInnerBlends is equal to BlendTime
-	UPROPERTY(EditAnywhere, BlueprintReadWrite, Category = Settings, meta = (PinHiddenByDefault))
-	float MaxTimeBeforeFreezingInnerBlends = -1.f;
-
 	// if AnimationTime and MaxAnimationDeltaTime are positive and the currently playing animation total time differs more than MaxAnimationDeltaTime from AnimationTime
 	// (animation desynchronized from the requested time) the blend stack will force a blend into the same animation
 	UPROPERTY(EditAnywhere, BlueprintReadWrite, Category = Settings, meta = (PinHiddenByDefault))
diff --git a/Engine/Plugins/Experimental/Animation/PoseSearch/Source/Runtime/Private/AnimNode_MotionMatching.cpp b/Engine/Plugins/Experimental/Animation/PoseSearch/Source/Runtime/Private/AnimNode_MotionMatching.cpp
--- a/Engine/Plugins/Experimental/Animation/PoseSearch/Source/Runtime/Private/AnimNode_MotionMatching.cpp
+++ b/Engine/Plugins/Experimental/Animation/PoseSearch/Source/Runtime/Private/AnimNode_MotionMatching.cpp
@@ -203,9 +203,8 @@ void FAnimNode_MotionMatching::UpdateAssetPlayer(const FAnimationUpdateContext&
 				// root bone blending needs to be immediate if MM node controls the offset between mesh component and root bone
 				const float RootBoneBlendTime = YawFromAnimationBlendRate < 0.f ? BlendTime : 0.f;
 				FAnimNode_BlendStack_Standalone::BlendTo(Context, DatabaseAsset->GetAnimationAsset(), MotionMatchingState.CurrentSearchResult.AssetTime,
-					SearchIndexAsset->IsLooping(), SearchIndexAsset->IsMirrored(), CurrentResultDatabase->Schema->MirrorDataTable.Get(),
-					BlendTime, RootBoneBlendTime, MaxTimeBeforeFreezingInnerBlends,
-					BlendProfile, BlendOption, bUseInertialBlend, SearchIndexAsset->GetBlendParameters(), MotionMatchingState.WantedPlayRate);
+					SearchIndexAsset->IsLooping(), SearchIndexAsset->IsMirrored(), CurrentResultDatabase->Schema->MirrorDataTable.Get(), BlendTime,
+					RootBoneBlendTime, BlendProfile, BlendOption, bUseInertialBlend, SearchIndexAsset->GetBlendParameters(), MotionMatchingState.WantedPlayRate);
 			}
 		}
 	}
diff --git a/Engine/Plugins/Experimental/Animation/PoseSearch/Source/Runtime/Public/PoseSearch/AnimNode_MotionMatching.h b/Engine/Plugins/Experimental/Animation/PoseSearch/Source/Runtime/Public/PoseSearch/AnimNode_MotionMatching.h
--- a/Engine/Plugins/Experimental/Animation/PoseSearch/Source/Runtime/Public/PoseSearch/AnimNode_MotionMatching.h
+++ b/Engine/Plugins/Experimental/Animation/PoseSearch/Source/Runtime/Public/PoseSearch/AnimNode_MotionMatching.h
@@ -64,10 +64,6 @@ private:
 	UPROPERTY(EditAnywhere, Category = Settings, meta = (PinHiddenByDefault, ClampMin="0"))
 	float BlendTime = 0.2f;
 
-	// Max time in seconds before freezing the blend ratio of all blends, but the most recently requested one. Negative values would imply MaxTimeBeforeFreezingInnerBlends is equal to BlendTime
-	UPROPERTY(EditAnywhere, Category = Settings, meta = (PinHiddenByDefault))
-	float MaxTimeBeforeFreezingInnerBlends = -1.f;
-
 	// Set Blend Profiles (editable in the skeleton) to determine how the blending is distributed among your character's bones. It could be used to differentiate between upper body and lower body to blend timing.
 	UPROPERTY(EditAnywhere, Category = Settings, meta = (PinHiddenByDefault, UseAsBlendProfile = true))
 	TObjectPtr<UBlendProfile> BlendProfile;
diff --git a/Engine/Plugins/Experimental/Chooser/Source/Chooser/Private/AnimNode_ChooserPlayer.cpp b/Engine/Plugins/Experimental/Chooser/Source/Chooser/Private/AnimNode_ChooserPlayer.cpp
--- a/Engine/Plugins/Experimental/Chooser/Source/Chooser/Private/AnimNode_ChooserPlayer.cpp
+++ b/Engine/Plugins/Experimental/Chooser/Source/Chooser/Private/AnimNode_ChooserPlayer.cpp
@@ -137,8 +137,7 @@ void FAnimNode_ChooserPlayer::UpdateAssetPlayer(const FAnimationUpdateContext& C
 		}
 		
 		FAnimNode_BlendStack_Standalone::BlendTo(Context, NewAsset, Settings.StartTime,
-			bLoop, Settings.bMirror, MirrorDataTable,
-			Settings.BlendTime, Settings.BlendTime, Settings.BlendTime,
+			bLoop, Settings.bMirror, MirrorDataTable, Settings.BlendTime, Settings.BlendTime,
 			Settings.BlendProfile, Settings.BlendOption, Settings.bUseInertialBlend, FVector::Zero(), Settings.PlaybackRate,
 			GetGroupName(), GetGroupRole(), GetGroupMethod());
 
