Disabled filtering for anim notifies states in FBlendStackAnimEventsFilterContext to not affect their behaviour until proper filtering support for them is added.

#rnx

diff --git a/Engine/Plugins/Experimental/Animation/BlendStack/Source/Runtime/Private/BlendStackAnimEventsFilterScope.cpp b/Engine/Plugins/Experimental/Animation/BlendStack/Source/Runtime/Private/BlendStackAnimEventsFilterScope.cpp
--- a/Engine/Plugins/Experimental/Animation/BlendStack/Source/Runtime/Private/BlendStackAnimEventsFilterScope.cpp
+++ b/Engine/Plugins/Experimental/Animation/BlendStack/Source/Runtime/Private/BlendStackAnimEventsFilterScope.cpp
@@ -13,17 +13,24 @@ namespace UE { namespace Anim {
 
 	bool FBlendStackAnimEventsFilterContext::ShouldFilterNotify(const FAnimNotifyEventReference& InNotifyEventRef) const
 	{
-		const FName NotifyName = InNotifyEventRef.GetNotify()->NotifyName;
+		const FAnimNotifyEvent* Notify = InNotifyEventRef.GetNotify();
+
+		// Only filter anim notifies
+		if (!Notify || Notify->NotifyStateClass)
+		{
+			return false;
+		}
 		
+		const FName NotifyName = Notify->NotifyName;
 		const bool bWasRecentlyFired = NotifyBanList->Contains(NotifyName);
 		const bool bWasAlreadyFiredThisTick = FiredNotifies->Contains(NotifyName);
-		
+	
 		if (!bWasAlreadyFiredThisTick && !bWasRecentlyFired)
 		{
 			FiredNotifies->Push(NotifyName);
 			return false;
 		}
-
+		
 		return true;
 	}
 
