Fix warning V567: The modification of the 'CurrentSamplePoseLink' variable is unsequenced relative to another operation on the same variable. This may lead to undefined behavior.
#rnx

diff --git a/Engine/Plugins/Experimental/Animation/BlendStack/Source/Runtime/Private/AnimNode_BlendStack.cpp b/Engine/Plugins/Experimental/Animation/BlendStack/Source/Runtime/Private/AnimNode_BlendStack.cpp
--- a/Engine/Plugins/Experimental/Animation/BlendStack/Source/Runtime/Private/AnimNode_BlendStack.cpp
+++ b/Engine/Plugins/Experimental/Animation/BlendStack/Source/Runtime/Private/AnimNode_BlendStack.cpp
@@ -836,7 +836,7 @@ int32 FAnimNode_BlendStack_Standalone::GetNextPoseLinkIndex()
 	}
 
 	const int32 NumPoseLinks = SampleGraphPoseLinks.Num();
-	CurrentSamplePoseLink = ++CurrentSamplePoseLink;
+	++CurrentSamplePoseLink;
 	if (CurrentSamplePoseLink == NumPoseLinks) { CurrentSamplePoseLink = 0; }
 
 	return CurrentSamplePoseLink;
