Fix for pops when returning inertial blend stack to relevancy
#rnx

diff --git a/Engine/Plugins/Experimental/Animation/BlendStack/Source/Runtime/Private/AnimNode_BlendStack.cpp b/Engine/Plugins/Experimental/Animation/BlendStack/Source/Runtime/Private/AnimNode_BlendStack.cpp
--- a/Engine/Plugins/Experimental/Animation/BlendStack/Source/Runtime/Private/AnimNode_BlendStack.cpp
+++ b/Engine/Plugins/Experimental/Animation/BlendStack/Source/Runtime/Private/AnimNode_BlendStack.cpp
@@ -783,6 +783,14 @@ void FAnimNode_BlendStack_Standalone::BlendTo(const FAnimationUpdateContext& Con
 	const UBlendProfile* BlendProfile, EAlphaBlendOption BlendOption, bool bUseInertialBlend, const FVector& BlendParameters, float PlayRate,
 	FName GroupName, EAnimGroupRole::Type GroupRole, EAnimSyncMethod GroupMethod)
 {
+	const bool bBlendStackIsEmpty = AnimPlayers.IsEmpty();
+
+	// If the blend stack is empty, we shouldn't blend. Pop into the requested pose.
+	if (bBlendStackIsEmpty)
+	{		
+		BlendTime = 0.0f;
+	}
+
 	if (bUseInertialBlend)
 	{
 		RequestInertialBlend(Context, BlendTime, BlendProfile, BlendOption);
@@ -791,7 +799,7 @@ void FAnimNode_BlendStack_Standalone::BlendTo(const FAnimationUpdateContext& Con
 
 	// If we don't add a new player, re-use the same graph...
 	int32 NewSamplePoseLinkIndex = CurrentSamplePoseLink;
-	if (!AnimPlayers.IsEmpty() && AnimPlayers[0].GetCurrentBlendInTime() < MaxBlendInTimeToOverrideAnimation)
+	if (!bBlendStackIsEmpty && AnimPlayers[0].GetCurrentBlendInTime() < MaxBlendInTimeToOverrideAnimation)
 	{
 		// replacing AnimPlayers[0] with this new BlendTo request
 		UE_LOG(LogBlendStack, Verbose, TEXT("FAnimNode_BlendStack_Standalone '%s' replaced by '%s' because blend time in is less than MaxBlendInTimeToOverrideAnimation (%.2f / %.2f)"), *AnimPlayers[0].GetAnimationName(), *GetNameSafe(AnimationAsset), AnimPlayers[0].GetCurrentBlendInTime(), MaxBlendInTimeToOverrideAnimation);
@@ -1026,4 +1034,4 @@ void FBlendStack_SampleGraphPoseLink::SetInputPosePlayer(FBlendStackAnimPlayer&
 	Player = &InPlayer;
 }
 
-#undef LOCTEXT_NAMESPACE
\ No newline at end of file
+#undef LOCTEXT_NAMESPACE
