motion matching - logging Warning(s) for dropped animplayers when bStoreBlendedPose is false (required for blendstack tuning)

diff --git a/Engine/Plugins/Experimental/Animation/BlendStack/Source/Runtime/Private/AnimNode_BlendStack.cpp b/Engine/Plugins/Experimental/Animation/BlendStack/Source/Runtime/Private/AnimNode_BlendStack.cpp
--- a/Engine/Plugins/Experimental/Animation/BlendStack/Source/Runtime/Private/AnimNode_BlendStack.cpp
+++ b/Engine/Plugins/Experimental/Animation/BlendStack/Source/Runtime/Private/AnimNode_BlendStack.cpp
@@ -512,7 +512,7 @@ void FAnimNode_BlendStack_Standalone::Evaluate_AnyThread(FPoseContext& Output)
 		{
 			EvaluateAndBlendPlayerByIndex(PlayerIndex);
 
-				// too many AnimPlayers! we don't have enough available blends to hold them all, so we accumulate the blended poses into Output / BlendedPoseContext.
+			// too many AnimPlayers! we don't have enough available blends to hold them all, so we accumulate the blended poses into Output / BlendedPoseContext.
 			PopLastAnimPlayer();
 		}
 
@@ -526,6 +526,13 @@ void FAnimNode_BlendStack_Standalone::Evaluate_AnyThread(FPoseContext& Output)
 				// We store Output / BlendedPoseContext into the last AnimPlayer, that will hold a static pose, no longer an animation playing.
 				AnimPlayers.Last().StorePoseContext(Output);
 			}
+#if !NO_LOGGING
+			// warning if we're dropping an animplayer with relevant (MaxBlendInTimeToOverrideAnimation) weight (GetBlendInPercentage)
+			else if (AnimPlayers.Last().GetBlendInPercentage() < (1.f - MaxBlendInTimeToOverrideAnimation))
+			{
+				UE_LOG(LogBlendStack, Warning, TEXT("FAnimNode_BlendStack_Standalone dropping animplayer with blend in at %.2f"), AnimPlayers.Last().GetBlendInPercentage());
+			}
+#endif // !NO_LOGGING
 		}
 
 		// Continue with the evaluation of the most significant AnimPlayer(s) with the associated graphs
