blend stack - fix for crash when rebuilding animgraph / blend stack reset, while PIE is running

diff --git a/Engine/Plugins/Experimental/Animation/BlendStack/Source/Runtime/Private/AnimNode_BlendStack.cpp b/Engine/Plugins/Experimental/Animation/BlendStack/Source/Runtime/Private/AnimNode_BlendStack.cpp
--- a/Engine/Plugins/Experimental/Animation/BlendStack/Source/Runtime/Private/AnimNode_BlendStack.cpp
+++ b/Engine/Plugins/Experimental/Animation/BlendStack/Source/Runtime/Private/AnimNode_BlendStack.cpp
@@ -186,8 +186,15 @@ void FBlendStackAnimPlayer::RestorePoseContext(FPoseContext& PoseContext) const
 		// Remap the pose manually in those cases.
 		if (PoseContext.Pose.GetBoneContainer().GetSerialNumber() == StoredBoneContainer.GetSerialNumber())
 		{
-			check(PoseContext.Pose.GetNumBones() == StoredBones.Num());
-			FMemory::Memcpy(PoseContext.Pose.GetMutableBones().GetData(), StoredBones.GetData(), sizeof(FTransform) * PoseContext.Pose.GetNumBones());
+			if (StoredBones.IsEmpty())
+			{
+				PoseContext.Pose.ResetToRefPose();
+			}
+			else
+			{
+				check(PoseContext.Pose.GetNumBones() == StoredBones.Num());
+				FMemory::Memcpy(PoseContext.Pose.GetMutableBones().GetData(), StoredBones.GetData(), sizeof(FTransform) * PoseContext.Pose.GetNumBones());
+			}
 		}
 		else
 		{
