Fix for MaxBlendInTimeToOverrideAnimation stomping an animation at full weight, and first-frame blends not setting duration of blend properly

diff --git a/Engine/Plugins/Experimental/Animation/BlendStack/Source/Runtime/Private/AnimNode_BlendStack.cpp b/Engine/Plugins/Experimental/Animation/BlendStack/Source/Runtime/Private/AnimNode_BlendStack.cpp
--- a/Engine/Plugins/Experimental/Animation/BlendStack/Source/Runtime/Private/AnimNode_BlendStack.cpp
+++ b/Engine/Plugins/Experimental/Animation/BlendStack/Source/Runtime/Private/AnimNode_BlendStack.cpp
@@ -803,17 +803,19 @@ void FAnimNode_BlendStack_Standalone::BlendTo(const FAnimationUpdateContext& Con
 	if (bBlendStackIsEmpty)
 	{		
 		BlendTime = 0.0f;
+		RootBoneBlendTime =  0.0f;
 	}
 
 	if (bUseInertialBlend)
 	{
 		RequestInertialBlend(Context, BlendTime, BlendProfile, BlendOption);
 		BlendTime = 0.0f;
+		RootBoneBlendTime = 0.0f;
 	}
 
 	// If we don't add a new player, re-use the same graph...
 	int32 NewSamplePoseLinkIndex = CurrentSamplePoseLink;
-	if (!bBlendStackIsEmpty && AnimPlayers[0].GetCurrentBlendInTime() < MaxBlendInTimeToOverrideAnimation)
+	if (!bBlendStackIsEmpty &&  AnimPlayers[0].GetBlendInPercentage() < 1.0f && AnimPlayers[0].GetCurrentBlendInTime() < MaxBlendInTimeToOverrideAnimation)
 	{
 		// replacing AnimPlayers[0] with this new BlendTo request
 		UE_LOG(LogBlendStack, Verbose, TEXT("FAnimNode_BlendStack_Standalone '%s' replaced by '%s' because blend time in is less than MaxBlendInTimeToOverrideAnimation (%.2f / %.2f)"), *AnimPlayers[0].GetAnimationName(), *GetNameSafe(AnimationAsset), AnimPlayers[0].GetCurrentBlendInTime(), MaxBlendInTimeToOverrideAnimation);
