blend stack
- add support for state machine in blend stack graph
- support for ActivationDelay in BlendTo
- introduced UBlendStackInputAnimNodeLibrary to expose FAnimNode_BlendStackInput properties

diff --git a/Engine/Plugins/Experimental/Animation/BlendStack/Source/Editor/Public/AnimationBlendStackGraph.h b/Engine/Plugins/Experimental/Animation/BlendStack/Source/Editor/Public/AnimationBlendStackGraph.h
--- a/Engine/Plugins/Experimental/Animation/BlendStack/Source/Editor/Public/AnimationBlendStackGraph.h
+++ b/Engine/Plugins/Experimental/Animation/BlendStack/Source/Editor/Public/AnimationBlendStackGraph.h
@@ -2,8 +2,7 @@
 
 #pragma once
 
-#include "EdGraph/EdGraph.h"
-
+#include "AnimationGraph.h"
 #include "AnimationBlendStackGraph.generated.h"
 
 class UObject;
@@ -13,7 +12,7 @@ class UObject;
 *	The input node is dynamically linked to the blend stack's sample pose.
 */
 UCLASS(MinimalAPI)
-class UAnimationBlendStackGraph : public UEdGraph
+class UAnimationBlendStackGraph : public UAnimationGraph
 {
 	GENERATED_BODY()
 
diff --git a/Engine/Plugins/Experimental/Animation/BlendStack/Source/Runtime/Private/AnimNode_BlendStack.cpp b/Engine/Plugins/Experimental/Animation/BlendStack/Source/Runtime/Private/AnimNode_BlendStack.cpp
--- a/Engine/Plugins/Experimental/Animation/BlendStack/Source/Runtime/Private/AnimNode_BlendStack.cpp
+++ b/Engine/Plugins/Experimental/Animation/BlendStack/Source/Runtime/Private/AnimNode_BlendStack.cpp
@@ -27,7 +27,7 @@ static constexpr FBoneIndexType RootBoneIndexType = 0;
 // FBlendStackAnimPlayer
 void FBlendStackAnimPlayer::Initialize(const FAnimationInitializeContext& Context, UAnimationAsset* AnimationAsset, float AccumulatedTime, bool bLoop,
 	bool bMirrored, UMirrorDataTable* MirrorDataTable, float BlendTime, float RootBoneBlendTime, 
-	const UBlendProfile* BlendProfile, EAlphaBlendOption InBlendOption, const FVector& BlendParameters, float PlayRate, int32 InPoseLinkIdx,
+	const UBlendProfile* BlendProfile, EAlphaBlendOption InBlendOption, const FVector& BlendParameters, float PlayRate, float ActivationDelay, int32 InPoseLinkIdx,
 	FName GroupName, EAnimGroupRole::Type GroupRole, EAnimSyncMethod GroupMethod)
 {
 	if (bMirrored && !MirrorDataTable)
@@ -81,6 +81,7 @@ void FBlendStackAnimPlayer::Initialize(const FAnimationInitializeContext& Contex
 
 	TotalBlendInTime = BlendTime;
 	CurrentBlendInTime = 0.f;
+	TimeToActivation = ActivationDelay;
 
 	MirrorNode.SetMirrorDataTable(MirrorDataTable);
 	MirrorNode.SetMirror(bMirrored);
@@ -315,10 +316,30 @@ float FBlendStackAnimPlayer::GetPlayRate() const
 	return 0.f;
 }
 
-void FBlendStackAnimPlayer::AdvanceBlendInTime(const float DeltaTime, int32 PlayerDepth, float PlayerDepthBlendInTimeMultiplier)
+bool FBlendStackAnimPlayer::IsActive() const
+{
+	return TimeToActivation <= 0.f;
+}
+
+void FBlendStackAnimPlayer::UpdateWithDeltaTime(float DeltaTime, int32 PlayerDepth, float PlayerDepthBlendInTimeMultiplier)
 {
 	const bool bIsMainPlayer = PlayerDepth == 0;
 
+	if (TimeToActivation > 0.f)
+	{
+		TimeToActivation -= DeltaTime;
+
+		if (TimeToActivation < 0.f)
+		{
+			DeltaTime = -TimeToActivation;
+			TimeToActivation = 0.f;
+		}
+		else
+		{
+			DeltaTime = 0.f;
+		}
+	}
+
 	if (bIsMainPlayer)
 	{
 		CurrentBlendInTime += DeltaTime;
@@ -455,8 +476,8 @@ void FAnimNode_BlendStack_Standalone::Evaluate_AnyThread(FPoseContext& Output)
 	for (int32 i = 0; i < AnimPlayers.Num(); ++i)
 	{
 		const FBlendStackAnimPlayer& AnimPlayer = AnimPlayers[i];
-		MessageBuilder.Appendf(TEXT("%d) t:%.2f/%.2f m:%d %s\n"), 
-			i, AnimPlayer.GetCurrentBlendInTime(), AnimPlayer.GetTotalBlendInTime(),
+		MessageBuilder.Appendf(TEXT("%d) t:%.2f/%.2f a:%.2f m:%d %s\n"), 
+			i, AnimPlayer.GetCurrentBlendInTime(), AnimPlayer.GetTotalBlendInTime(), AnimPlayer.GetTimeToActivation(),
 			AnimPlayer.GetMirror() ? 1 : 0, *AnimPlayer.GetAnimationName());
 	}
 
@@ -630,22 +651,23 @@ void FAnimNode_BlendStack_Standalone::UpdateAssetPlayer(const FAnimationUpdateCo
 	float CurrentWeightMultiplier = 1.f;
 	const int32 BlendStackSize = AnimPlayers.Num();
 	int32 AnimPlayerIndex = 0;
-	for (; AnimPlayerIndex < BlendStackSize; ++AnimPlayerIndex)
+	while (AnimPlayerIndex < BlendStackSize)
 	{
 		FBlendStackAnimPlayer& AnimPlayer = AnimPlayers[AnimPlayerIndex];
 		const bool bIsLastAnimPlayers = AnimPlayerIndex == BlendStackSize - 1;
 		const float BlendInPercentage = bIsLastAnimPlayers ? 1.f : AnimPlayer.GetBlendInPercentage();
 		const float AnimPlayerBlendWeight = CurrentWeightMultiplier * BlendInPercentage;
 
-		// don't break for AnimPlayerIndex == 0 since FAnimNode_BlendStack_Standalone::BlendTo initialize the AnimPlayer with a weight of zero
-		if (AnimPlayerIndex > 0 && AnimPlayerBlendWeight < UE_KINDA_SMALL_NUMBER)
-		{
-			break;
-		}
-
 		FAnimationUpdateContext AnimPlayerContext = Context.FractionalWeightAndRootMotion(AnimPlayerBlendWeight, AnimPlayerBlendWeight);
 		UpdateSample((AnimPlayerIndex == 0) ? AnimPlayerContext : AnimPlayerContext.AsInactive(), AnimPlayerIndex);
 		CurrentWeightMultiplier *= (1.f - BlendInPercentage);
+
+		++AnimPlayerIndex;
+
+		if (CurrentWeightMultiplier < UE_KINDA_SMALL_NUMBER)
+		{
+			break;
+		}
 	}
 
 	// AnimPlayers[AnimPlayerIndex] is the first FBlendStackAnimPlayer with a weight contribution of zero,
@@ -709,22 +731,25 @@ void FAnimNode_BlendStack_Standalone::UpdateSample(const FAnimationUpdateContext
 {
 	FBlendStackAnimPlayer& SamplePlayer = AnimPlayers[PlayerIndex];
 
-	const bool bHasSampleGraph = IsSampleGraphAvailableForPlayer(PlayerIndex);
-	if (bHasSampleGraph)
+	if (SamplePlayer.IsActive())
 	{
-		FBlendStack_SampleGraphPoseLink& PoseLink = SampleGraphPoseLinks[SamplePlayer.GetPoseLinkIndex()];
-		PoseLink.SetInputPosePlayer(SamplePlayer);
-		// The anim player may or may not have its Update_AnyThread called through the graph update. 
-		PoseLink.Root.Update(Context);
-	}
-	else
-	{
-		// If we have no sample graph, update the player directly.
-		SamplePlayer.Update_AnyThread(Context);
+		const bool bHasSampleGraph = IsSampleGraphAvailableForPlayer(PlayerIndex);
+		if (bHasSampleGraph)
+		{
+			FBlendStack_SampleGraphPoseLink& PoseLink = SampleGraphPoseLinks[SamplePlayer.GetPoseLinkIndex()];
+			PoseLink.SetInputPosePlayer(SamplePlayer);
+			// The anim player may or may not have its Update_AnyThread called through the graph update. 
+			PoseLink.Root.Update(Context);
+		}
+		else
+		{
+			// If we have no sample graph, update the player directly.
+			SamplePlayer.Update_AnyThread(Context);
+		}
 	}
 
 	// Advance the blend-in time regardless of whether or not the player was updated.
-	SamplePlayer.AdvanceBlendInTime(Context.GetDeltaTime(), PlayerIndex, PlayerDepthBlendInTimeMultiplier);
+	SamplePlayer.UpdateWithDeltaTime(Context.GetDeltaTime(), PlayerIndex, PlayerDepthBlendInTimeMultiplier);
 }
 
 void FAnimNode_BlendStack_Standalone::CacheBonesForSample(const FAnimationCacheBonesContext& Context, const int32 PlayerIndex)
@@ -794,7 +819,7 @@ static void RequestInertialBlend(const FAnimationUpdateContext& Context, float B
 
 void FAnimNode_BlendStack_Standalone::BlendTo(const FAnimationUpdateContext& Context, UAnimationAsset* AnimationAsset, float AccumulatedTime, bool bLoop,
 	bool bMirrored, UMirrorDataTable* MirrorDataTable, float BlendTime, float RootBoneBlendTime, 
-	const UBlendProfile* BlendProfile, EAlphaBlendOption BlendOption, bool bUseInertialBlend, const FVector& BlendParameters, float PlayRate,
+	const UBlendProfile* BlendProfile, EAlphaBlendOption BlendOption, bool bUseInertialBlend, const FVector& BlendParameters, float PlayRate, float ActivationDelay,
 	FName GroupName, EAnimGroupRole::Type GroupRole, EAnimSyncMethod GroupMethod)
 {
 	const bool bBlendStackIsEmpty = AnimPlayers.IsEmpty();
@@ -815,7 +840,9 @@ void FAnimNode_BlendStack_Standalone::BlendTo(const FAnimationUpdateContext& Con
 
 	// If we don't add a new player, re-use the same graph...
 	int32 NewSamplePoseLinkIndex = CurrentSamplePoseLink;
-	if (!bBlendStackIsEmpty &&  AnimPlayers[0].GetBlendInPercentage() < 1.0f && AnimPlayers[0].GetCurrentBlendInTime() < MaxBlendInTimeToOverrideAnimation)
+	if (!bBlendStackIsEmpty &&  AnimPlayers[0].GetBlendInPercentage() < 1.0f &&
+		AnimPlayers[0].GetCurrentBlendInTime() < MaxBlendInTimeToOverrideAnimation &&
+		FMath::IsNearlyEqual(AnimPlayers[0].GetTimeToActivation(), ActivationDelay))
 	{
 		// replacing AnimPlayers[0] with this new BlendTo request
 		UE_LOG(LogBlendStack, Verbose, TEXT("FAnimNode_BlendStack_Standalone '%s' replaced by '%s' because blend time in is less than MaxBlendInTimeToOverrideAnimation (%.2f / %.2f)"), *AnimPlayers[0].GetAnimationName(), *GetNameSafe(AnimationAsset), AnimPlayers[0].GetCurrentBlendInTime(), MaxBlendInTimeToOverrideAnimation);
@@ -835,7 +862,7 @@ void FAnimNode_BlendStack_Standalone::BlendTo(const FAnimationUpdateContext& Con
 	FBlendStackAnimPlayer& AnimPlayer = AnimPlayers[0];
 
 	FAnimationInitializeContext InitContext(Context.AnimInstanceProxy, Context.SharedContext);
-	AnimPlayer.Initialize(InitContext, AnimationAsset, AccumulatedTime, bLoop, bMirrored, MirrorDataTable, BlendTime, RootBoneBlendTime, BlendProfile, BlendOption, BlendParameters, PlayRate, NewSamplePoseLinkIndex, GroupName, GroupRole, GroupMethod);
+	AnimPlayer.Initialize(InitContext, AnimationAsset, AccumulatedTime, bLoop, bMirrored, MirrorDataTable, BlendTime, RootBoneBlendTime, BlendProfile, BlendOption, BlendParameters, PlayRate, ActivationDelay, NewSamplePoseLinkIndex, GroupName, GroupRole, GroupMethod);
 	InitializeSample(InitContext, AnimPlayer);
 }
 
@@ -881,8 +908,8 @@ void FAnimNode_BlendStack_Standalone::GatherDebugData(FNodeDebugData& DebugData)
 	for (int32 i = 0; i < AnimPlayers.Num(); ++i)
 	{
 		const FBlendStackAnimPlayer& AnimPlayer = AnimPlayers[i];
-		DebugData.AddDebugItem(FString::Printf(TEXT("%d) t:%.2f/%.2f m:%d %s"),
-				i, AnimPlayer.GetCurrentBlendInTime(), AnimPlayer.GetTotalBlendInTime(),
+		DebugData.AddDebugItem(FString::Printf(TEXT("%d) t:%.2f/%.2f a:%.2f m:%d %s"),
+				i, AnimPlayer.GetCurrentBlendInTime(), AnimPlayer.GetTotalBlendInTime(), AnimPlayer.GetTimeToActivation(),
 				AnimPlayer.GetMirror() ? 1 : 0, *AnimPlayer.GetAnimationName()));
 	}
 #endif // ENABLE_ANIM_DEBUG
@@ -995,7 +1022,11 @@ void FAnimNode_BlendStack::UpdateAssetPlayer(const FAnimationUpdateContext& Cont
 	GetEvaluateGraphExposedInputs().Execute(Context);
 
 	bool bExecuteBlendTo = false;
-	if (AnimPlayers.IsEmpty())
+	if (AnimationAsset == nullptr)
+	{
+		bExecuteBlendTo = false;
+	}
+	else if (AnimPlayers.IsEmpty())
 	{
 		bExecuteBlendTo = true;
 	}
diff --git a/Engine/Plugins/Experimental/Animation/BlendStack/Source/Runtime/Private/BlendStackAnimNodeLibrary.cpp b/Engine/Plugins/Experimental/Animation/BlendStack/Source/Runtime/Private/BlendStackAnimNodeLibrary.cpp
--- a/Engine/Plugins/Experimental/Animation/BlendStack/Source/Runtime/Private/BlendStackAnimNodeLibrary.cpp
+++ b/Engine/Plugins/Experimental/Animation/BlendStack/Source/Runtime/Private/BlendStackAnimNodeLibrary.cpp
@@ -17,42 +17,41 @@ void UBlendStackAnimNodeLibrary::BlendTo(const FAnimUpdateContext& Context,
 										bool bMirrored,
 										float BlendTime,
 										FVector BlendParameters,
-										float WantedPlayRate)
+										float WantedPlayRate,
+										float ActivationDelay)
 {
-	if (FAnimNode_BlendStack* BlendStackNodePtr = BlendStackNode.GetAnimNodePtr<FAnimNode_BlendStack>())
+	if (AnimationAsset != nullptr)
 	{
-		if (const FAnimationUpdateContext* AnimationUpdateContext = Context.GetContext())
+		if (FAnimNode_BlendStack* BlendStackNodePtr = BlendStackNode.GetAnimNodePtr<FAnimNode_BlendStack>())
 		{
-			if (AnimationAsset == nullptr)
+			if (const FAnimationUpdateContext* AnimationUpdateContext = Context.GetContext())
 			{
-				UE_LOG(LogBlendStack, Warning, TEXT("UBlendStackAnimNodeLibrary::BlendTo called with null animation asset."));
-				return;
+				BlendStackNodePtr->BlendTo(
+					*AnimationUpdateContext,
+					AnimationAsset,
+					AnimationTime,
+					bLoop,
+					bMirrored,
+					BlendStackNodePtr->MirrorDataTable,
+					BlendTime,
+					BlendStackNodePtr->RootBoneBlendTime,
+					BlendStackNodePtr->BlendProfile,
+					BlendStackNodePtr->BlendOption,
+					BlendStackNodePtr->bUseInertialBlend,
+					BlendParameters,
+					WantedPlayRate,
+					ActivationDelay);
+			}
+			else
+			{
+				UE_LOG(LogBlendStack, Warning, TEXT("UBlendStackAnimNodeLibrary::BlendTo called with an invalid context."));
 			}
-
-			BlendStackNodePtr->BlendTo(
-				*AnimationUpdateContext, 
-				AnimationAsset, 
-				AnimationTime, 
-				bLoop,
-				bMirrored,
-				BlendStackNodePtr->MirrorDataTable,
-				BlendTime,
-				BlendStackNodePtr->RootBoneBlendTime,
-				BlendStackNodePtr->BlendProfile,
-				BlendStackNodePtr->BlendOption,
-				BlendStackNodePtr->bUseInertialBlend,
-				BlendParameters,
-				WantedPlayRate);
 		}
 		else
 		{
-			UE_LOG(LogBlendStack, Warning, TEXT("UBlendStackAnimNodeLibrary::BlendTo called with an invalid context."));
+			UE_LOG(LogBlendStack, Warning, TEXT("UBlendStackAnimNodeLibrary::BlendTo called with an invalid type."));
 		}
 	}
-	else
-	{
-		UE_LOG(LogBlendStack, Warning, TEXT("UBlendStackAnimNodeLibrary::BlendTo called with an invalid type."));
-	}
 }
 
 void UBlendStackAnimNodeLibrary::ForceBlendNextUpdate(const FBlendStackAnimNodeReference& BlendStackNode)
diff --git a/Engine/Plugins/Experimental/Animation/BlendStack/Source/Runtime/Private/BlendStackInputAnimNodeLibrary.cpp b/Engine/Plugins/Experimental/Animation/BlendStack/Source/Runtime/Private/BlendStackInputAnimNodeLibrary.cpp
new file mode 100644
--- /dev/null
+++ b/Engine/Plugins/Experimental/Animation/BlendStack/Source/Runtime/Private/BlendStackInputAnimNodeLibrary.cpp
@@ -0,0 +1,28 @@
+// Copyright Epic Games, Inc. All Rights Reserved.
+
+#include "BlendStack/BlendStackInputAnimNodeLibrary.h"
+#include "BlendStack/AnimNode_BlendStackInput.h"
+#include "BlendStack/AnimNode_BlendStack.h"
+
+FBlendStackInputAnimNodeReference UBlendStackInputAnimNodeLibrary::ConvertToBlendStackInputNode(const FAnimNodeReference& Node, EAnimNodeReferenceConversionResult& Result)
+{
+	return FAnimNodeReference::ConvertToType<FBlendStackInputAnimNodeReference>(Node, Result);
+}
+
+void UBlendStackInputAnimNodeLibrary::GetProperties(const FBlendStackInputAnimNodeReference& BlendStackInputNode, UAnimationAsset*& AnimationAsset, float& AccumulatedTime)
+{
+	AnimationAsset = nullptr;
+	AccumulatedTime = 0.f;
+
+	if (const FAnimNode_BlendStackInput* BlendStackInputNodePtr = BlendStackInputNode.GetAnimNodePtr<FAnimNode_BlendStackInput>())
+	{
+		if (BlendStackInputNodePtr->Player)
+		{
+			if (const FBlendStackAnimPlayer* Player = *BlendStackInputNodePtr->Player)
+			{
+				AnimationAsset = Player->GetAnimationAsset();
+				AccumulatedTime = Player->GetAccumulatedTime();
+			}
+		}
+	}
+}
diff --git a/Engine/Plugins/Experimental/Animation/BlendStack/Source/Runtime/Public/BlendStack/AnimNode_BlendStack.h b/Engine/Plugins/Experimental/Animation/BlendStack/Source/Runtime/Public/BlendStack/AnimNode_BlendStack.h
--- a/Engine/Plugins/Experimental/Animation/BlendStack/Source/Runtime/Public/BlendStack/AnimNode_BlendStack.h
+++ b/Engine/Plugins/Experimental/Animation/BlendStack/Source/Runtime/Public/BlendStack/AnimNode_BlendStack.h
@@ -15,7 +15,7 @@ struct BLENDSTACK_API FBlendStackAnimPlayer
 	
 	void Initialize(const FAnimationInitializeContext& Context, UAnimationAsset* AnimationAsset, float AccumulatedTime, bool bLoop,
 		bool bMirrored, UMirrorDataTable* MirrorDataTable, float BlendTime, float RootBoneBlendTime,
-		const UBlendProfile* BlendProfile, EAlphaBlendOption InBlendOption, const FVector& BlendParameters, float PlayRate, int32 InPoseLinkIdx,
+		const UBlendProfile* BlendProfile, EAlphaBlendOption InBlendOption, const FVector& BlendParameters, float PlayRate, float ActivationDelay, int32 InPoseLinkIdx,
 		FName GroupName, EAnimGroupRole::Type GroupRole, EAnimSyncMethod Method);
 	
 	void UpdatePlayRate(float PlayRate);
@@ -36,7 +36,9 @@ struct BLENDSTACK_API FBlendStackAnimPlayer
 
 	float GetTotalBlendInTime() const { return TotalBlendInTime; }
 	float GetCurrentBlendInTime() const { return CurrentBlendInTime; }
-	void AdvanceBlendInTime(const float DeltaTime, int32 PlayerDepth, float PlayerDepthBlendInTimeMultiplier);
+	float GetTimeToActivation() const { return TimeToActivation; }
+	
+	void UpdateWithDeltaTime(float DeltaTime, int32 PlayerDepth, float PlayerDepthBlendInTimeMultiplier);
 	bool GetMirror() const { return MirrorNode.GetMirror(); }
 	FVector GetBlendParameters() const;
 	void SetBlendParameters(const FVector& BlendParameters);
@@ -49,7 +51,8 @@ struct BLENDSTACK_API FBlendStackAnimPlayer
 	void RestorePoseContext(FPoseContext& PoseContext) const;
 	void UpdateSourceLinkNode();
 	bool IsLooping() const;
-	
+	bool IsActive() const;
+
 	// Curves to add to the pose after the player evaluates
 	TBaseBlendedCurve<FDefaultAllocator, UE::Anim::FCurveElement> OverrideCurve;
 
@@ -82,6 +85,7 @@ private:
 
 	float TotalBlendInTime = 0.f;
 	float CurrentBlendInTime = 0.f;
+	float TimeToActivation = 0.f;
 };
 
 USTRUCT(BlueprintInternalUseOnly)
@@ -129,7 +133,7 @@ struct BLENDSTACK_API FAnimNode_BlendStack_Standalone : public FAnimNode_AssetPl
 	void BlendTo(const FAnimationUpdateContext& Context, UAnimationAsset* AnimationAsset, float AccumulatedTime = 0.f, bool bLoop = false, 
 		bool bMirrored = false, UMirrorDataTable* MirrorDataTable = nullptr, float BlendTime = 0.2f, float RootBoneBlendTime = -1.f,
 		const UBlendProfile* BlendProfile = nullptr, EAlphaBlendOption BlendOption = EAlphaBlendOption::Linear, 
-		bool bUseInertialBlend = false, const FVector& BlendParameters = FVector::Zero(), float PlayRate = 1.f,
+		bool bUseInertialBlend = false, const FVector& BlendParameters = FVector::Zero(), float PlayRate = 1.f, float ActivationDelay = 0.f,
 		FName GroupName = NAME_None, EAnimGroupRole::Type GroupRole = EAnimGroupRole::CanBeLeader, EAnimSyncMethod Method = EAnimSyncMethod::DoNotSync);
 	void UpdatePlayRate(float PlayRate);
 	void Reset();
diff --git a/Engine/Plugins/Experimental/Animation/BlendStack/Source/Runtime/Public/BlendStack/BlendStackAnimNodeLibrary.h b/Engine/Plugins/Experimental/Animation/BlendStack/Source/Runtime/Public/BlendStack/BlendStackAnimNodeLibrary.h
--- a/Engine/Plugins/Experimental/Animation/BlendStack/Source/Runtime/Public/BlendStack/BlendStackAnimNodeLibrary.h
+++ b/Engine/Plugins/Experimental/Animation/BlendStack/Source/Runtime/Public/BlendStack/BlendStackAnimNodeLibrary.h
@@ -49,5 +49,6 @@ public:
 						bool bMirrored = false,
 						float BlendTime = 0.2f,
 						FVector BlendParameters = FVector::ZeroVector,
-						float WantedPlayRate = 1.f);
+						float WantedPlayRate = 1.f,
+						float ActivationDelay = 0.f);
 };
\ No newline at end of file
diff --git a/Engine/Plugins/Experimental/Animation/BlendStack/Source/Runtime/Public/BlendStack/BlendStackInputAnimNodeLibrary.h b/Engine/Plugins/Experimental/Animation/BlendStack/Source/Runtime/Public/BlendStack/BlendStackInputAnimNodeLibrary.h
new file mode 100644
--- /dev/null
+++ b/Engine/Plugins/Experimental/Animation/BlendStack/Source/Runtime/Public/BlendStack/BlendStackInputAnimNodeLibrary.h
@@ -0,0 +1,42 @@
+// Copyright Epic Games, Inc. All Rights Reserved.
+
+#pragma once
+
+#include "Animation/AnimExecutionContext.h"
+#include "Animation/AnimNodeReference.h"
+#include "Kismet/BlueprintFunctionLibrary.h"
+#include "BlendStackInputAnimNodeLibrary.generated.h"
+
+struct FAnimNode_BlendStackInput;
+
+USTRUCT(BlueprintType)
+struct FBlendStackInputAnimNodeReference : public FAnimNodeReference
+{
+	GENERATED_BODY()
+
+	typedef FAnimNode_BlendStackInput FInternalNodeType;
+};
+
+// Exposes operations that can be run on a Blend Stack node via Anim Node Functions such as "On Become Relevant" and "On Update".
+UCLASS(Experimental)
+class BLENDSTACK_API UBlendStackInputAnimNodeLibrary : public UBlueprintFunctionLibrary
+{
+	GENERATED_BODY()
+
+public:
+	/** Get a blend stack input node context from an anim node context */
+	UFUNCTION(BlueprintCallable, Category = "Animation|BlendStack", meta = (BlueprintThreadSafe, ExpandEnumAsExecs = "Result"))
+	static FBlendStackInputAnimNodeReference ConvertToBlendStackInputNode(const FAnimNodeReference& Node, EAnimNodeReferenceConversionResult& Result);
+
+	/** Get a blend stack input node context from an anim node context (pure) */
+	UFUNCTION(BlueprintPure, Category = "Animation|BlendStack", meta = (BlueprintThreadSafe, DisplayName = "Convert to Blend Stack Input Node"))
+	static void ConvertToBlendStackInputNodePure(const FAnimNodeReference& Node, FBlendStackInputAnimNodeReference& BlendStackInputNode, bool& Result)
+	{
+		EAnimNodeReferenceConversionResult ConversionResult;
+		BlendStackInputNode = ConvertToBlendStackInputNode(Node, ConversionResult);
+		Result = (ConversionResult == EAnimNodeReferenceConversionResult::Succeeded);
+	}
+
+	UFUNCTION(BlueprintCallable, Category = "Animation|BlendStack", meta = (BlueprintThreadSafe))
+	static void GetProperties(const FBlendStackInputAnimNodeReference& BlendStackInputNode, UAnimationAsset*& AnimationAsset, float& AccumulatedTime);
+};
\ No newline at end of file
diff --git a/Engine/Plugins/Experimental/Chooser/Source/Chooser/Private/AnimNode_ChooserPlayer.cpp b/Engine/Plugins/Experimental/Chooser/Source/Chooser/Private/AnimNode_ChooserPlayer.cpp
--- a/Engine/Plugins/Experimental/Chooser/Source/Chooser/Private/AnimNode_ChooserPlayer.cpp
+++ b/Engine/Plugins/Experimental/Chooser/Source/Chooser/Private/AnimNode_ChooserPlayer.cpp
@@ -160,7 +160,7 @@ void FAnimNode_ChooserPlayer::UpdateAssetPlayer(const FAnimationUpdateContext& C
 		
 		FAnimNode_BlendStack_Standalone::BlendTo(Context, NewAsset, Settings.StartTime,
 			bLoop, Settings.bMirror, MirrorDataTable, Settings.BlendTime, Settings.BlendTime,
-			Settings.BlendProfile, Settings.BlendOption, Settings.bUseInertialBlend, FVector::Zero(), Settings.PlaybackRate,
+			Settings.BlendProfile, Settings.BlendOption, Settings.bUseInertialBlend, FVector::Zero(), Settings.PlaybackRate, 0.f,
 			GetGroupName(), GetGroupRole(), GetGroupMethod());
 
 		if (!Settings.CurveOverrides.Values.IsEmpty())
