Anim Node Reference Tag support for nodes inside BlendStack subgraphs

diff --git a/Engine/Plugins/Experimental/Animation/BlendStack/Source/Editor/Private/AnimGraphNode_BlendStack.cpp b/Engine/Plugins/Experimental/Animation/BlendStack/Source/Editor/Private/AnimGraphNode_BlendStack.cpp
--- a/Engine/Plugins/Experimental/Animation/BlendStack/Source/Editor/Private/AnimGraphNode_BlendStack.cpp
+++ b/Engine/Plugins/Experimental/Animation/BlendStack/Source/Editor/Private/AnimGraphNode_BlendStack.cpp
@@ -1,15 +1,16 @@
 // Copyright Epic Games, Inc. All Rights Reserved.
 
 #include "AnimGraphNode_BlendStack.h"
-#include "Animation/AnimRootMotionProvider.h"
-#include "AnimationBlendStackGraphSchema.h"
 #include "AnimationBlendStackGraph.h"
-#include "Kismet2/BlueprintEditorUtils.h"
-#include "Kismet2/KismetEditorUtilities.h"
+#include "AnimationBlendStackGraphSchema.h"
+#include "AnimGraphNode_BlendStackInput.h"
+#include "AnimGraphNode_Root.h"
 #include "EdGraphUtilities.h"
 #include "IAnimBlueprintCompilationContext.h"
-#include "AnimGraphNode_Root.h"
-#include "AnimGraphNode_BlendStackInput.h"
+#include "K2Node_AnimNodeReference.h"
+#include "Animation/AnimRootMotionProvider.h"
+#include "Kismet2/BlueprintEditorUtils.h"
+#include "Kismet2/KismetEditorUtilities.h"
 
 #define LOCTEXT_NAMESPACE "AnimGraphNode_BlendStack"
 
@@ -49,6 +50,7 @@ void UAnimGraphNode_BlendStack::BakeDataDuringCompilation(class FCompilerResults
 }
 
 void UAnimGraphNode_BlendStack_Base::ExpandGraphAndProcessNodes(
+	int GraphIndex,
 	UEdGraph* SourceGraph, 
 	UAnimGraphNode_Base* SourceRootNode, TArrayView<UAnimGraphNode_BlendStackInput*> InputNodes,
 	IAnimBlueprintCompilationContext& InCompilationContext, IAnimBlueprintGeneratedClassCompiledData& OutCompiledData,
@@ -72,6 +74,8 @@ void UAnimGraphNode_BlendStack_Base::ExpandGraphAndProcessNodes(
 	{
 		SourceInputObjects[Index] = InCompilationContext.GetMessageLog().FindSourceObject(InputNodes[Index]);
 	}
+	
+	TSet<FName> TagsToRemap;
 
 	for (auto NodeIt = ClonedGraph->Nodes.CreateIterator(); NodeIt; ++NodeIt)
 	{
@@ -79,6 +83,14 @@ void UAnimGraphNode_BlendStack_Base::ExpandGraphAndProcessNodes(
 		if (UAnimGraphNode_Base* TestNode = Cast<UAnimGraphNode_Base>(ClonedNode))
 		{
 			AnimNodeList.Add(TestNode);
+			
+			FName NodeTag = TestNode->GetTag();
+			if (!NodeTag.IsNone())
+			{
+				TagsToRemap.Add(NodeTag);
+				NodeTag.SetNumber(GraphIndex + 1);
+				TestNode->SetTag(NodeTag);
+			}
 
 			//@TODO: There ought to be a better way to determine this
 			UObject* TestObject = InCompilationContext.GetMessageLog().FindSourceObject(TestNode);
@@ -97,11 +109,20 @@ void UAnimGraphNode_BlendStack_Base::ExpandGraphAndProcessNodes(
 
 	check(OutRootNode && !OutInputNodes.Contains(nullptr));
 
-	// Run another expansion pass to catch the graph we just added (this is slightly wasteful)
-	InCompilationContext.ExpansionStep(ClonedGraph, false);
+	for (auto NodeIt = ClonedGraph->Nodes.CreateIterator(); NodeIt; ++NodeIt)
+	{
+		UEdGraphNode* ClonedNode = *NodeIt;
 
-	// Validate graph now we have expanded/pruned
-	InCompilationContext.ValidateGraphIsWellFormed(ClonedGraph);
+		if (UK2Node_AnimNodeReference* ReferenceNode = Cast<UK2Node_AnimNodeReference>(ClonedNode))
+		{
+			FName NodeTag = ReferenceNode->GetTag();
+			if (TagsToRemap.Contains(NodeTag))
+			{
+				NodeTag.SetNumber(GraphIndex + 1);
+				ReferenceNode->SetTag(NodeTag);
+			}
+		}
+	}
 
 	// Move the cloned nodes into the consolidated event graph
 	const bool bIsLoading = InCompilationContext.GetBlueprint()->bIsRegeneratingOnLoad || IsAsyncLoading();
@@ -116,6 +137,12 @@ void UAnimGraphNode_BlendStack_Base::ExpandGraphAndProcessNodes(
 		InCompilationContext.PruneIsolatedAnimationNodes(RootSet, AnimNodeList);
 		InCompilationContext.ProcessAnimationNodes(AnimNodeList);
 	}
+	
+	// Run another expansion pass to catch the graph we just added (this is slightly wasteful)
+	InCompilationContext.ExpansionStep(ClonedGraph, false);
+	
+	// Validate graph now we have expanded/pruned
+	InCompilationContext.ValidateGraphIsWellFormed(ClonedGraph);
 }
 
 void UAnimGraphNode_BlendStack_Base::OnProcessDuringCompilation(IAnimBlueprintCompilationContext& InCompilationContext, IAnimBlueprintGeneratedClassCompiledData& OutCompiledData)
@@ -153,7 +180,7 @@ void UAnimGraphNode_BlendStack_Base::OnProcessDuringCompilation(IAnimBlueprintCo
 		TArray<UAnimGraphNode_BlendStackInput*> ClonedInputNodes;
 		ClonedInputNodes.SetNum(InputNodes.Num());
 
-		ExpandGraphAndProcessNodes(SampleGraph, SampleGraph->ResultNode, InputNodes, InCompilationContext, OutCompiledData, ClonedRootNode, ClonedInputNodes);
+		ExpandGraphAndProcessNodes(Index, SampleGraph, SampleGraph->ResultNode, InputNodes, InCompilationContext, OutCompiledData, ClonedRootNode, ClonedInputNodes);
 
 		// Blend stack node is potentially nested in the struct, so we can't use FPoseLinkMappingRecord. Patch at runtime instead.
 		for (UAnimGraphNode_BlendStackInput* InputNode : ClonedInputNodes)
diff --git a/Engine/Plugins/Experimental/Animation/BlendStack/Source/Editor/Public/AnimGraphNode_BlendStack.h b/Engine/Plugins/Experimental/Animation/BlendStack/Source/Editor/Public/AnimGraphNode_BlendStack.h
--- a/Engine/Plugins/Experimental/Animation/BlendStack/Source/Editor/Public/AnimGraphNode_BlendStack.h
+++ b/Engine/Plugins/Experimental/Animation/BlendStack/Source/Editor/Public/AnimGraphNode_BlendStack.h
@@ -34,6 +34,7 @@ class BLENDSTACKEDITOR_API UAnimGraphNode_BlendStack_Base : public UAnimGraphNod
 protected:
 	// Helper function for compilation
 	void ExpandGraphAndProcessNodes(
+		int GraphIndex,
 		UEdGraph* SourceGraph, 
 		UAnimGraphNode_Base* SourceRootNode, TArrayView<UAnimGraphNode_BlendStackInput*> SourceInputNode,
 		IAnimBlueprintCompilationContext& InCompilationContext, IAnimBlueprintGeneratedClassCompiledData& OutCompiledData,
diff --git a/Engine/Source/Editor/AnimGraph/Public/K2Node_AnimNodeReference.h b/Engine/Source/Editor/AnimGraph/Public/K2Node_AnimNodeReference.h
--- a/Engine/Source/Editor/AnimGraph/Public/K2Node_AnimNodeReference.h
+++ b/Engine/Source/Editor/AnimGraph/Public/K2Node_AnimNodeReference.h
@@ -29,6 +29,11 @@ public:
 	// Get the text used for the node's label
 	FText GetLabelText() const;
 	
+	// Get the tag for this node, if any
+	FName GetTag() const { return Tag; }
+
+	// Set the tag for this node
+	void SetTag(FName InTag) { Tag = InTag; }
 private:
 	/** The node tag we reference */
 	UPROPERTY()
