Fix for blend stack crashing after graph change during recompilation
#rnx

diff --git a/Engine/Plugins/Experimental/Animation/BlendStack/Source/Editor/Private/AnimGraphNode_BlendStack.cpp b/Engine/Plugins/Experimental/Animation/BlendStack/Source/Editor/Private/AnimGraphNode_BlendStack.cpp
--- a/Engine/Plugins/Experimental/Animation/BlendStack/Source/Editor/Private/AnimGraphNode_BlendStack.cpp
+++ b/Engine/Plugins/Experimental/Animation/BlendStack/Source/Editor/Private/AnimGraphNode_BlendStack.cpp
@@ -162,20 +162,26 @@ void UAnimGraphNode_BlendStack_Base::OnProcessDuringCompilation(IAnimBlueprintCo
 	{
 		// Input Pose is connected to Output Pose, so the sample graph does nothing.
 		// No need to use allocate sample graphs in that case.
-		AnimNode->SampleGraphPoseLinks.Reset();
+		AnimNode->PerSampleGraphPoseLinks.Reset();
 		return;
 	}
 
 	// Allocate one sample graph per-active blend plus an extra one for the stored pose.
-	AnimNode->SampleGraphPoseLinks.SetNum(MaxBlendsNum + 1);
+	const int32 NumSampleGraphs = MaxBlendsNum + 1;
+	AnimNode->PerSampleGraphPoseLinks.SetNum(NumSampleGraphs);
 
 	TArray<UAnimGraphNode_BlendStackInput*> InputNodes;
 	BoundGraph->GetNodesOfClass<UAnimGraphNode_BlendStackInput>(InputNodes);
 
 	int32 BlendStackAllocationIndex = InCompilationContext.GetAllocationIndexOfNode(this);
-	for(int32 Index = 0; Index < AnimNode->SampleGraphPoseLinks.Num(); ++ Index)
+
+	FStructProperty* NodeProperty = GetFNodeProperty();
+	check(NodeProperty);
+	FArrayProperty* PoseLinksProperty = CastFieldChecked<FArrayProperty>(NodeProperty->Struct->FindPropertyByName(GET_MEMBER_NAME_CHECKED(FAnimNode_BlendStack_Standalone, PerSampleGraphPoseLinks)));
+	check(PoseLinksProperty);
+
+	for(int32 Index = 0; Index < NumSampleGraphs; ++Index)
 	{
-		FBlendStack_SampleGraphPoseLink& GraphPoseLink = AnimNode->SampleGraphPoseLinks[Index];
 		UAnimGraphNode_Base* ClonedRootNode;
 		TArray<UAnimGraphNode_BlendStackInput*> ClonedInputNodes;
 		ClonedInputNodes.SetNum(InputNodes.Num());
@@ -188,7 +194,8 @@ void UAnimGraphNode_BlendStack_Base::OnProcessDuringCompilation(IAnimBlueprintCo
 			InputNode->Node.BlendStackAllocationIndex = BlendStackAllocationIndex;
 			InputNode->Node.SampleIndex = Index;
 		}
-		GraphPoseLink.RootNodeIndex = InCompilationContext.GetAllocationIndexOfNode(ClonedRootNode);
+
+		InCompilationContext.AddPoseLinkMappingRecord(FPoseLinkMappingRecord::MakeFromArrayEntry(this, ClonedRootNode, PoseLinksProperty, Index));
 	}
 }
 
diff --git a/Engine/Plugins/Experimental/Animation/BlendStack/Source/Runtime/Private/AnimNode_BlendStack.cpp b/Engine/Plugins/Experimental/Animation/BlendStack/Source/Runtime/Private/AnimNode_BlendStack.cpp
--- a/Engine/Plugins/Experimental/Animation/BlendStack/Source/Runtime/Private/AnimNode_BlendStack.cpp
+++ b/Engine/Plugins/Experimental/Animation/BlendStack/Source/Runtime/Private/AnimNode_BlendStack.cpp
@@ -583,20 +583,12 @@ void FAnimNode_BlendStack_Standalone::Initialize_AnyThread(const FAnimationIniti
 	
 	Reset();
 
-	if (SampleGraphPoseLinks.IsEmpty() == false)
+	if (PerSampleGraphPoseLinks.IsEmpty() == false)
 	{
-		IAnimClassInterface* AnimBlueprintClass = Context.GetAnimClass();
-		check(AnimBlueprintClass);
-
-		// Patch our pose links
-		for (FBlendStack_SampleGraphPoseLink& GraphPoseLink : SampleGraphPoseLinks)
+		SampleGraphExecutionHelpers.SetNum(PerSampleGraphPoseLinks.Num());
+		for (UE::BlendStack::FBlendStack_SampleGraphExecutionHelper& ExecutionHelper : SampleGraphExecutionHelpers)
 		{
-			if (GraphPoseLink.RootNodeIndex != INDEX_NONE)
-			{
-				GraphPoseLink.Root.LinkID = AnimBlueprintClass->GetAnimNodeProperties().Num() - 1 - GraphPoseLink.RootNodeIndex;
-			}
-
-			GraphPoseLink.CacheBoneCounter.Reset();
+			ExecutionHelper.CacheBoneCounter.Reset();
 		}
 	}
 }
@@ -685,7 +677,7 @@ bool FAnimNode_BlendStack_Standalone::IsSampleGraphAvailableForPlayer(const int3
 {
 	// If we have any sample graphs, our player has been assigned a pose link index.
 	// Players with a stored pose don't need to run the graph.
-	return !SampleGraphPoseLinks.IsEmpty() && !AnimPlayers[PlayerIndex].HasValidPoseContext();
+	return !PerSampleGraphPoseLinks.IsEmpty() && !AnimPlayers[PlayerIndex].HasValidPoseContext();
 }
 
 void FAnimNode_BlendStack_Standalone::EvaluateSample(FPoseContext& Output, const int32 PlayerIndex)
@@ -700,21 +692,22 @@ void FAnimNode_BlendStack_Standalone::EvaluateSample(FPoseContext& Output, const
 		return;
 	}
 
-	FBlendStack_SampleGraphPoseLink& PoseLink = SampleGraphPoseLinks[SamplePlayer.GetPoseLinkIndex()];
-	PoseLink.EvaluatePlayer(Output, SamplePlayer);
+	const int32 SampleIndex = SamplePlayer.GetPoseLinkIndex();
+	UE::BlendStack::FBlendStack_SampleGraphExecutionHelper& PoseLink = SampleGraphExecutionHelpers[SampleIndex];
+	PoseLink.EvaluatePlayer(Output, SamplePlayer, PerSampleGraphPoseLinks[SampleIndex]);
 }
 
-void FBlendStack_SampleGraphPoseLink::EvaluatePlayer(FPoseContext& Output, FBlendStackAnimPlayer& SamplePlayer)
+void UE::BlendStack::FBlendStack_SampleGraphExecutionHelper::EvaluatePlayer(FPoseContext& Output, FBlendStackAnimPlayer& SamplePlayer, FPoseLink& SamplePoseLink)
 {
 	SetInputPosePlayer(SamplePlayer);
 
 	// Make sure CacheBones has been called before evaluating.
-	ConditionalCacheBones(Output);
+	ConditionalCacheBones(Output, SamplePoseLink);
 	// The anim player may or may not have its Evaluate_AnyThread called through the graph update. 
-	Root.Evaluate(Output);
+	SamplePoseLink.Evaluate(Output);
 }
 
-void FBlendStack_SampleGraphPoseLink::ConditionalCacheBones(const FAnimationBaseContext& Context)
+void UE::BlendStack::FBlendStack_SampleGraphExecutionHelper::ConditionalCacheBones(const FAnimationBaseContext& Context, FPoseLink& SamplePoseLink)
 {
 	// Only call CacheBones when needed.
 	if (!CacheBoneCounter.IsSynchronized_Counter(Context.AnimInstanceProxy->GetCachedBonesCounter()))
@@ -723,7 +716,7 @@ void FBlendStack_SampleGraphPoseLink::ConditionalCacheBones(const FAnimationBase
 		CacheBoneCounter.SynchronizeWith(Context.AnimInstanceProxy->GetCachedBonesCounter());
 
 		FAnimationCacheBonesContext CacheBoneContext(Context.AnimInstanceProxy);
-		Root.CacheBones(CacheBoneContext);
+		SamplePoseLink.CacheBones(CacheBoneContext);
 	}
 }
 
@@ -736,10 +729,11 @@ void FAnimNode_BlendStack_Standalone::UpdateSample(const FAnimationUpdateContext
 		const bool bHasSampleGraph = IsSampleGraphAvailableForPlayer(PlayerIndex);
 		if (bHasSampleGraph)
 		{
-			FBlendStack_SampleGraphPoseLink& PoseLink = SampleGraphPoseLinks[SamplePlayer.GetPoseLinkIndex()];
-			PoseLink.SetInputPosePlayer(SamplePlayer);
+			const int32 SampleIndex = SamplePlayer.GetPoseLinkIndex();
+			UE::BlendStack::FBlendStack_SampleGraphExecutionHelper& ExecutionHelper = SampleGraphExecutionHelpers[SampleIndex];
+			ExecutionHelper.SetInputPosePlayer(SamplePlayer);
 			// The anim player may or may not have its Update_AnyThread called through the graph update. 
-			PoseLink.Root.Update(Context);
+			PerSampleGraphPoseLinks[SampleIndex].Update(Context);
 		}
 		else
 		{
@@ -758,8 +752,9 @@ void FAnimNode_BlendStack_Standalone::CacheBonesForSample(const FAnimationCacheB
 	const bool bHasSampleGraph = IsSampleGraphAvailableForPlayer(PlayerIndex);
 	if (bHasSampleGraph)
 	{
-		FBlendStack_SampleGraphPoseLink& PoseLink = SampleGraphPoseLinks[SamplePlayer.GetPoseLinkIndex()];
-		PoseLink.ConditionalCacheBones(Context);
+		const int32 SampleIndex = SamplePlayer.GetPoseLinkIndex();
+		UE::BlendStack::FBlendStack_SampleGraphExecutionHelper& ExecutionHelper = SampleGraphExecutionHelpers[SampleIndex];
+		ExecutionHelper.ConditionalCacheBones(Context, PerSampleGraphPoseLinks[SampleIndex]);
 	}
 }
 
@@ -767,10 +762,12 @@ void FAnimNode_BlendStack_Standalone::InitializeSample(const FAnimationInitializ
 {
 	if (SamplePlayer.GetPoseLinkIndex() != INDEX_NONE)
 	{
-		FBlendStack_SampleGraphPoseLink& PoseLink = SampleGraphPoseLinks[SamplePlayer.GetPoseLinkIndex()];
-		PoseLink.SetInputPosePlayer(SamplePlayer);
-		PoseLink.Root.Initialize(Context);
-		PoseLink.ConditionalCacheBones(Context);
+		const int32 SampleIndex = SamplePlayer.GetPoseLinkIndex();
+		FPoseLink& PoseLink = PerSampleGraphPoseLinks[SampleIndex];
+		UE::BlendStack::FBlendStack_SampleGraphExecutionHelper& ExecutionHelper = SampleGraphExecutionHelpers[SampleIndex];
+		ExecutionHelper.SetInputPosePlayer(SamplePlayer);
+		PoseLink.Initialize(Context);
+		ExecutionHelper.ConditionalCacheBones(Context, PoseLink);
 	}
 }
 
@@ -881,12 +878,12 @@ void FAnimNode_BlendStack_Standalone::Reset()
 
 int32 FAnimNode_BlendStack_Standalone::GetNextPoseLinkIndex()
 {
-	if (SampleGraphPoseLinks.IsEmpty())
+	if (PerSampleGraphPoseLinks.IsEmpty())
 	{
 		return INDEX_NONE;
 	}
 
-	const int32 NumPoseLinks = SampleGraphPoseLinks.Num();
+	const int32 NumPoseLinks = PerSampleGraphPoseLinks.Num();
 	++CurrentSamplePoseLink;
 	if (CurrentSamplePoseLink == NumPoseLinks) { CurrentSamplePoseLink = 0; }
 
@@ -1074,7 +1071,7 @@ void FAnimNode_BlendStack::ForceBlendNextUpdate()
 	bForceBlendNextUpdate = true;
 }
 
-void FBlendStack_SampleGraphPoseLink::SetInputPosePlayer(FBlendStackAnimPlayer& InPlayer)
+void UE::BlendStack::FBlendStack_SampleGraphExecutionHelper::SetInputPosePlayer(FBlendStackAnimPlayer& InPlayer)
 {
 	// Because our anim players may get reallocated, or change indices due to push/pops,
 	// we must call this before every operation that might end up needing the anim player through the graph's input nodes.
diff --git a/Engine/Plugins/Experimental/Animation/BlendStack/Source/Runtime/Private/AnimNode_BlendStackInput.cpp b/Engine/Plugins/Experimental/Animation/BlendStack/Source/Runtime/Private/AnimNode_BlendStackInput.cpp
--- a/Engine/Plugins/Experimental/Animation/BlendStack/Source/Runtime/Private/AnimNode_BlendStackInput.cpp
+++ b/Engine/Plugins/Experimental/Animation/BlendStack/Source/Runtime/Private/AnimNode_BlendStackInput.cpp
@@ -36,7 +36,7 @@ void FAnimNode_BlendStackInput::Update_AnyThread(const FAnimationUpdateContext&
 		const int32 BlendStackNodeIndex = AnimBlueprintClass->GetAnimNodeProperties().Num() - 1 - BlendStackAllocationIndex;
 		check(Context.AnimInstanceProxy);
 		FAnimNode_BlendStack_Standalone* BlendStackNode = GetBlendStackNodeFromIndex(Context, AnimBlueprintClass, BlendStackNodeIndex);
-		Player = &BlendStackNode->SampleGraphPoseLinks[SampleIndex].Player;
+		Player = &BlendStackNode->SampleGraphExecutionHelpers[SampleIndex].Player;
 	}
 
 	check(Player && *Player)
diff --git a/Engine/Plugins/Experimental/Animation/BlendStack/Source/Runtime/Public/BlendStack/AnimNode_BlendStack.h b/Engine/Plugins/Experimental/Animation/BlendStack/Source/Runtime/Public/BlendStack/AnimNode_BlendStack.h
--- a/Engine/Plugins/Experimental/Animation/BlendStack/Source/Runtime/Public/BlendStack/AnimNode_BlendStack.h
+++ b/Engine/Plugins/Experimental/Animation/BlendStack/Source/Runtime/Public/BlendStack/AnimNode_BlendStack.h
@@ -88,32 +88,29 @@ private:
 	float TimeToActivation = 0.f;
 };
 
-USTRUCT(BlueprintInternalUseOnly)
-struct FBlendStack_SampleGraphPoseLink
+namespace UE::BlendStack
+{
+struct FBlendStack_SampleGraphExecutionHelper
 {
-	GENERATED_BODY()
-
-	UPROPERTY()
-	int32 RootNodeIndex = INDEX_NONE;
-
-	UPROPERTY(Transient)
-	FPoseLink Root;
-
 	FBlendStackAnimPlayer* Player = nullptr;
 	FGraphTraversalCounter CacheBoneCounter;
 
 	void SetInputPosePlayer(FBlendStackAnimPlayer& InPlayer);
-	void EvaluatePlayer(FPoseContext& Output, FBlendStackAnimPlayer& SamplePlayer);
-	void ConditionalCacheBones(const FAnimationBaseContext& Output);
+	void EvaluatePlayer(FPoseContext& Output, FBlendStackAnimPlayer& SamplePlayer, FPoseLink& SamplePoseLink);
+	void ConditionalCacheBones(const FAnimationBaseContext& Output, FPoseLink& SamplePoseLink);
 };
+} // namespace UE::BlendStack
 
 USTRUCT(BlueprintInternalUseOnly)
 struct BLENDSTACK_API FAnimNode_BlendStack_Standalone : public FAnimNode_AssetPlayerBase
 {
 	GENERATED_BODY()
 
+	TArray<UE::BlendStack::FBlendStack_SampleGraphExecutionHelper> SampleGraphExecutionHelpers;
+
 	UPROPERTY()
-	TArray<FBlendStack_SampleGraphPoseLink> SampleGraphPoseLinks;
+	TArray<FPoseLink> PerSampleGraphPoseLinks;
+
 	int32 CurrentSamplePoseLink = -1;
 
 	UPROPERTY(Transient)
