motion matching - removing YawFromAnimationTrajectoryBlendTime, YawFromAnimationBlendRate (FAnimNode_OffsetRootBone is a more generic way of controlling the root bone via playing animation)

diff --git a/Engine/Plugins/Animation/BlendStack/Source/Runtime/Private/AnimNode_BlendStack.cpp b/Engine/Plugins/Animation/BlendStack/Source/Runtime/Private/AnimNode_BlendStack.cpp
--- a/Engine/Plugins/Animation/BlendStack/Source/Runtime/Private/AnimNode_BlendStack.cpp
+++ b/Engine/Plugins/Animation/BlendStack/Source/Runtime/Private/AnimNode_BlendStack.cpp
@@ -27,9 +27,8 @@ static constexpr FBoneIndexType RootBoneIndexType = 0;
 /////////////////////////////////////////////////////
 // FBlendStackAnimPlayer
 void FBlendStackAnimPlayer::Initialize(const FAnimationInitializeContext& Context, UAnimationAsset* AnimationAsset, float AccumulatedTime, bool bLoop,
-	bool bMirrored, UMirrorDataTable* MirrorDataTable, float BlendTime, float RootBoneBlendTime, 
-	const UBlendProfile* BlendProfile, EAlphaBlendOption InBlendOption, const FVector& BlendParameters, float PlayRate, float ActivationDelay, int32 InPoseLinkIdx,
-	FName GroupName, EAnimGroupRole::Type GroupRole, EAnimSyncMethod GroupMethod)
+	bool bMirrored, UMirrorDataTable* MirrorDataTable, float BlendTime, const UBlendProfile* BlendProfile, EAlphaBlendOption InBlendOption,
+	const FVector& BlendParameters, float PlayRate, float ActivationDelay, int32 InPoseLinkIdx, FName GroupName, EAnimGroupRole::Type GroupRole, EAnimSyncMethod GroupMethod)
 {
 	if (bMirrored && !MirrorDataTable)
 	{
@@ -41,7 +40,6 @@ void FBlendStackAnimPlayer::Initialize(const FAnimationInitializeContext& Contex
 	check(Skeleton);
 
 	const FReferenceSkeleton& RefSkeleton = Skeleton->GetReferenceSkeleton();
-	const bool bApplyDifferentRootBoneBlendTime = RootBoneBlendTime >= 0.f && !FMath::IsNearlyEqual(RootBoneBlendTime, BlendTime);
 	const int32 NumSkeletonBones = RefSkeleton.GetNum();
 	if (NumSkeletonBones <= 0)
 	{
@@ -53,29 +51,9 @@ void FBlendStackAnimPlayer::Initialize(const FAnimationInitializeContext& Contex
 		if (BlendProfile != nullptr)
 		{
 			TotalBlendInTimePerBone.Init(BlendTime, NumSkeletonBones);
-
 			BlendProfile->FillSkeletonBoneDurationsArray(TotalBlendInTimePerBone, BlendTime, Skeleton);
-
-			if (bApplyDifferentRootBoneBlendTime)
-			{
-				TotalBlendInTimePerBone[RootBoneIndexType] *= RootBoneBlendTime / BlendTime;
-			}
-
 			BlendTime = *Algo::MaxElement(TotalBlendInTimePerBone);
 		}
-		else if (bApplyDifferentRootBoneBlendTime)
-		{
-			TotalBlendInTimePerBone.Init(BlendTime, NumSkeletonBones);
-			TotalBlendInTimePerBone[RootBoneIndexType] *= RootBoneBlendTime / BlendTime;
-			BlendTime = FMath::Max(BlendTime, RootBoneBlendTime);
-		}
-	}
-	else if (bApplyDifferentRootBoneBlendTime)
-	{
-		// handling BlendTime ~= 0 and RootBoneBlendTime >= 0
-		TotalBlendInTimePerBone.Init(BlendTime, NumSkeletonBones);
-		TotalBlendInTimePerBone[RootBoneIndexType] = RootBoneBlendTime;
-		BlendTime = FMath::Max(BlendTime, RootBoneBlendTime);
 	}
 
 	BlendOption = InBlendOption;
@@ -816,9 +794,8 @@ static void RequestInertialBlend(const FAnimationUpdateContext& Context, float B
 }
 
 void FAnimNode_BlendStack_Standalone::BlendTo(const FAnimationUpdateContext& Context, UAnimationAsset* AnimationAsset, float AccumulatedTime, bool bLoop,
-	bool bMirrored, UMirrorDataTable* MirrorDataTable, float BlendTime, float RootBoneBlendTime, 
-	const UBlendProfile* BlendProfile, EAlphaBlendOption BlendOption, bool bUseInertialBlend, const FVector& BlendParameters, float PlayRate, float ActivationDelay,
-	FName GroupName, EAnimGroupRole::Type GroupRole, EAnimSyncMethod GroupMethod)
+	bool bMirrored, UMirrorDataTable* MirrorDataTable, float BlendTime, const UBlendProfile* BlendProfile, EAlphaBlendOption BlendOption, bool bUseInertialBlend,
+	const FVector& BlendParameters, float PlayRate, float ActivationDelay, FName GroupName, EAnimGroupRole::Type GroupRole, EAnimSyncMethod GroupMethod)
 {
 	const bool bBlendStackIsEmpty = AnimPlayers.IsEmpty();
 
@@ -826,14 +803,12 @@ void FAnimNode_BlendStack_Standalone::BlendTo(const FAnimationUpdateContext& Con
 	if (bBlendStackIsEmpty)
 	{		
 		BlendTime = 0.0f;
-		RootBoneBlendTime =  0.0f;
 	}
 
 	if (bUseInertialBlend)
 	{
 		RequestInertialBlend(Context, BlendTime, BlendProfile, BlendOption);
 		BlendTime = 0.0f;
-		RootBoneBlendTime = 0.0f;
 	}
 
 	// If we don't add a new player, re-use the same graph...
@@ -860,7 +835,7 @@ void FAnimNode_BlendStack_Standalone::BlendTo(const FAnimationUpdateContext& Con
 	FBlendStackAnimPlayer& AnimPlayer = AnimPlayers[0];
 
 	FAnimationInitializeContext InitContext(Context.AnimInstanceProxy, Context.SharedContext);
-	AnimPlayer.Initialize(InitContext, AnimationAsset, AccumulatedTime, bLoop, bMirrored, MirrorDataTable, BlendTime, RootBoneBlendTime, BlendProfile, BlendOption, BlendParameters, PlayRate, ActivationDelay, NewSamplePoseLinkIndex, GroupName, GroupRole, GroupMethod);
+	AnimPlayer.Initialize(InitContext, AnimationAsset, AccumulatedTime, bLoop, bMirrored, MirrorDataTable, BlendTime, BlendProfile, BlendOption, BlendParameters, PlayRate, ActivationDelay, NewSamplePoseLinkIndex, GroupName, GroupRole, GroupMethod);
 	InitializeSample(InitContext, AnimPlayer);
 }
 
@@ -1058,8 +1033,8 @@ void FAnimNode_BlendStack::UpdateAssetPlayer(const FAnimationUpdateContext& Cont
 
 	if (bExecuteBlendTo)
 	{
-		BlendTo(Context, AnimationAsset, AnimationTime, bLoop, bMirrored, MirrorDataTable.Get(), BlendTime,
-			RootBoneBlendTime, BlendProfile, BlendOption, bUseInertialBlend, BlendParameters, WantedPlayRate);
+		BlendTo(Context, AnimationAsset, AnimationTime, bLoop, bMirrored, MirrorDataTable.Get(),
+			BlendTime, BlendProfile, BlendOption, bUseInertialBlend, BlendParameters, WantedPlayRate);
 	}
 
 	const bool bDidBlendToRequestAnInertialBlend = bExecuteBlendTo && bUseInertialBlend;
diff --git a/Engine/Plugins/Animation/BlendStack/Source/Runtime/Private/BlendStackAnimNodeLibrary.cpp b/Engine/Plugins/Animation/BlendStack/Source/Runtime/Private/BlendStackAnimNodeLibrary.cpp
--- a/Engine/Plugins/Animation/BlendStack/Source/Runtime/Private/BlendStackAnimNodeLibrary.cpp
+++ b/Engine/Plugins/Animation/BlendStack/Source/Runtime/Private/BlendStackAnimNodeLibrary.cpp
@@ -59,7 +59,6 @@ void UBlendStackAnimNodeLibrary::BlendTo(const FAnimUpdateContext& Context,
 					bMirrored,
 					BlendStackNodePtr->MirrorDataTable,
 					BlendTime,
-					BlendStackNodePtr->RootBoneBlendTime,
 					BlendStackNodePtr->BlendProfile,
 					BlendStackNodePtr->BlendOption,
 					BlendStackNodePtr->bUseInertialBlend,
diff --git a/Engine/Plugins/Animation/BlendStack/Source/Runtime/Public/BlendStack/AnimNode_BlendStack.h b/Engine/Plugins/Animation/BlendStack/Source/Runtime/Public/BlendStack/AnimNode_BlendStack.h
--- a/Engine/Plugins/Animation/BlendStack/Source/Runtime/Public/BlendStack/AnimNode_BlendStack.h
+++ b/Engine/Plugins/Animation/BlendStack/Source/Runtime/Public/BlendStack/AnimNode_BlendStack.h
@@ -14,9 +14,8 @@ struct BLENDSTACK_API FBlendStackAnimPlayer
 	GENERATED_BODY()
 	
 	void Initialize(const FAnimationInitializeContext& Context, UAnimationAsset* AnimationAsset, float AccumulatedTime, bool bLoop,
-		bool bMirrored, UMirrorDataTable* MirrorDataTable, float BlendTime, float RootBoneBlendTime,
-		const UBlendProfile* BlendProfile, EAlphaBlendOption InBlendOption, const FVector& BlendParameters, float PlayRate, float ActivationDelay, int32 InPoseLinkIdx,
-		FName GroupName, EAnimGroupRole::Type GroupRole, EAnimSyncMethod Method);
+		bool bMirrored, UMirrorDataTable* MirrorDataTable, float BlendTime, const UBlendProfile* BlendProfile, EAlphaBlendOption InBlendOption,
+		const FVector& BlendParameters, float PlayRate, float ActivationDelay, int32 InPoseLinkIdx, FName GroupName, EAnimGroupRole::Type GroupRole, EAnimSyncMethod Method);
 	
 	void UpdatePlayRate(float PlayRate);
 	void Evaluate_AnyThread(FPoseContext& Output);
@@ -128,7 +127,7 @@ struct BLENDSTACK_API FAnimNode_BlendStack_Standalone : public FAnimNode_AssetPl
 	// End of FAnimNode_Base interface
 
 	void BlendTo(const FAnimationUpdateContext& Context, UAnimationAsset* AnimationAsset, float AccumulatedTime = 0.f, bool bLoop = false, 
-		bool bMirrored = false, UMirrorDataTable* MirrorDataTable = nullptr, float BlendTime = 0.2f, float RootBoneBlendTime = -1.f,
+		bool bMirrored = false, UMirrorDataTable* MirrorDataTable = nullptr, float BlendTime = 0.2f,
 		const UBlendProfile* BlendProfile = nullptr, EAlphaBlendOption BlendOption = EAlphaBlendOption::Linear, 
 		bool bUseInertialBlend = false, const FVector& BlendParameters = FVector::Zero(), float PlayRate = 1.f, float ActivationDelay = 0.f,
 		FName GroupName = NAME_None, EAnimGroupRole::Type GroupRole = EAnimGroupRole::CanBeLeader, EAnimSyncMethod Method = EAnimSyncMethod::DoNotSync);
@@ -216,10 +215,6 @@ struct BLENDSTACK_API FAnimNode_BlendStack : public FAnimNode_BlendStack_Standal
 	UPROPERTY(EditAnywhere, BlueprintReadWrite, Category = Settings, meta = (PinHiddenByDefault, ClampMin = "0"))
 	float BlendTime = 0.2f;
 
-	// Time in seconds to blend out to the new pose root bone. Negative values would imply RootBoneBlendTime is equal to BlendTime
-	UPROPERTY(EditAnywhere, BlueprintReadWrite, Category = Settings, meta = (PinHiddenByDefault))
-	float RootBoneBlendTime = -1.f;
-
 	// if AnimationTime and MaxAnimationDeltaTime are positive and the currently playing animation total time differs more than MaxAnimationDeltaTime from AnimationTime
 	// (animation desynchronized from the requested time) the blend stack will force a blend into the same animation
 	UPROPERTY(EditAnywhere, BlueprintReadWrite, Category = Settings, meta = (PinHiddenByDefault))
diff --git a/Engine/Plugins/Animation/PoseSearch/Source/Runtime/Private/AnimNode_MotionMatching.cpp b/Engine/Plugins/Animation/PoseSearch/Source/Runtime/Private/AnimNode_MotionMatching.cpp
--- a/Engine/Plugins/Animation/PoseSearch/Source/Runtime/Private/AnimNode_MotionMatching.cpp
+++ b/Engine/Plugins/Animation/PoseSearch/Source/Runtime/Private/AnimNode_MotionMatching.cpp
@@ -176,8 +176,6 @@ void FAnimNode_MotionMatching::UpdateAssetPlayer(const FAnimationUpdateContext&
 		SearchThrottleTime,
 		PlayRate,
 		MotionMatchingState,
-		YawFromAnimationBlendRate,
-		YawFromAnimationTrajectoryBlendTime,
 		NextUpdateInterruptMode,
 		bShouldSearch,
 		bShouldUseCachedChannelData
@@ -201,11 +199,9 @@ void FAnimNode_MotionMatching::UpdateAssetPlayer(const FAnimationUpdateContext&
 
 			if (UAnimationAsset* AnimationAsset = Cast<UAnimationAsset>(DatabaseAsset->GetAnimationAsset()))
 			{
-				// root bone blending needs to be immediate if MM node controls the offset between mesh component and root bone
-				const float RootBoneBlendTime = YawFromAnimationBlendRate < 0.f ? BlendTime : 0.f;
 				FAnimNode_BlendStack_Standalone::BlendTo(Context, AnimationAsset, MotionMatchingState.CurrentSearchResult.AssetTime,
 					SearchIndexAsset->IsLooping(), SearchIndexAsset->IsMirrored(), CurrentResultDatabase->Schema->GetMirrorDataTable(DefaultRole), BlendTime,
-					RootBoneBlendTime, BlendProfile, BlendOption, bUseInertialBlend, SearchIndexAsset->GetBlendParameters(), MotionMatchingState.WantedPlayRate);
+					BlendProfile, BlendOption, bUseInertialBlend, SearchIndexAsset->GetBlendParameters(), MotionMatchingState.WantedPlayRate);
 			}
 			else
 			{
diff --git a/Engine/Plugins/Animation/PoseSearch/Source/Runtime/Private/PoseSearchLibrary.cpp b/Engine/Plugins/Animation/PoseSearch/Source/Runtime/Private/PoseSearchLibrary.cpp
--- a/Engine/Plugins/Animation/PoseSearch/Source/Runtime/Private/PoseSearchLibrary.cpp
+++ b/Engine/Plugins/Animation/PoseSearch/Source/Runtime/Private/PoseSearchLibrary.cpp
@@ -217,33 +217,6 @@ void FMotionMatchingState::UpdateWantedPlayRate(const UE::PoseSearch::FSearchCon
 	}
 }
 
-void FMotionMatchingState::UpdateRootBoneControl(const FAnimationUpdateContext& Context, float YawFromAnimationBlendRate)
-{
-	const FAnimInstanceProxy* AnimInstanceProxy = Context.AnimInstanceProxy;
-
-	const float CurrentComponentWorldYaw = FRotator(AnimInstanceProxy->GetComponentTransform().GetRotation()).Yaw;
-	if (YawFromAnimationBlendRate < 0.f)
-	{
-		ComponentWorldYaw = CurrentComponentWorldYaw;
-		ComponentDeltaYaw = 0.f;
-	}
-	else
-	{
-		// integrating MotionMatchingState.ComponentWorldYaw with a lerped value between 
-		// the previous frame AnimationDeltaYaw and CurrentComponentDeltaYaw (component delta yaw happened during the delta time)
-		const float CurrentComponentDeltaYaw = FRotator::NormalizeAxis(CurrentComponentWorldYaw - ComponentWorldYaw);
-
-		// lerping the animation delta with the capsule delta
-		const float LerpValue = FMath::Min(YawFromAnimationBlendRate * Context.GetDeltaTime(), 1.f);
-		const float LerpedDeltaYaw = FMath::Lerp(AnimationDeltaYaw, CurrentComponentDeltaYaw, LerpValue);
-
-		ComponentWorldYaw = FRotator::NormalizeAxis(ComponentWorldYaw + LerpedDeltaYaw);
-
-		// @todo: handle the case when the character is on top of a rotating platform
-		ComponentDeltaYaw = FRotator::NormalizeAxis(ComponentWorldYaw - CurrentComponentWorldYaw);
-	}
-}
-
 #if UE_POSE_SEARCH_TRACE_ENABLED
 void UPoseSearchLibrary::TraceMotionMatchingState(
 	UE::PoseSearch::FSearchContext& SearchContext,
@@ -378,8 +351,6 @@ void UPoseSearchLibrary::UpdateMotionMatchingState(
 	float SearchThrottleTime,
 	const FFloatInterval& PlayRate,
 	FMotionMatchingState& InOutMotionMatchingState,
-	float YawFromAnimationBlendRate,
-	float YawFromAnimationTrajectoryBlendTime,
 	EPoseSearchInterruptMode InterruptMode,
 	bool bShouldSearch,
 	bool bShouldUseCachedChannelData,
@@ -400,32 +371,15 @@ void UPoseSearchLibrary::UpdateMotionMatchingState(
 		return;
 	}
 
-	InOutMotionMatchingState.UpdateRootBoneControl(Context, YawFromAnimationBlendRate);
-
 	const float DeltaTime = Context.GetDeltaTime();
 
 	InOutMotionMatchingState.bJumpedToPose = false;
 
 	// used when YawFromAnimationBlendRate is greater than zero, by setting a future (YawFromAnimationTrajectoryBlendTime seconds ahead) root bone to the skeleton default
-	FMemStackPoseHistory MemStackPoseHistory;
 	const IPoseHistory* PoseHistory = nullptr;
 	if (IPoseHistoryProvider* PoseHistoryProvider = Context.GetMessage<IPoseHistoryProvider>())
 	{
 		PoseHistory = &PoseHistoryProvider->GetPoseHistory();
-
-		if (YawFromAnimationBlendRate >= 0.f && YawFromAnimationTrajectoryBlendTime > 0.f)
-		{
-			MemStackPoseHistory.Init(PoseHistory);
-
-			const FTransform& RefRootBoneTransform = Context.AnimInstanceProxy->GetSkeleton()->GetReferenceSkeleton().GetRefBonePose()[RootBoneIndexType];
-			MemStackPoseHistory.AddFutureRootBone(YawFromAnimationTrajectoryBlendTime, RefRootBoneTransform, false);
-
-#if ENABLE_DRAW_DEBUG && ENABLE_ANIM_DEBUG
-			MemStackPoseHistory.DebugDraw(*Context.AnimInstanceProxy, FColor::Emerald);
-#endif // ENABLE_DRAW_DEBUG && ENABLE_ANIM_DEBUG
-
-			PoseHistory = &MemStackPoseHistory;
-		}
 	}
 
 	FMemMark Mark(FMemStack::Get());
diff --git a/Engine/Plugins/Animation/PoseSearch/Source/Runtime/Public/PoseSearch/AnimNode_MotionMatching.h b/Engine/Plugins/Animation/PoseSearch/Source/Runtime/Public/PoseSearch/AnimNode_MotionMatching.h
--- a/Engine/Plugins/Animation/PoseSearch/Source/Runtime/Public/PoseSearch/AnimNode_MotionMatching.h
+++ b/Engine/Plugins/Animation/PoseSearch/Source/Runtime/Public/PoseSearch/AnimNode_MotionMatching.h
@@ -96,17 +96,6 @@ private:
 	UPROPERTY(EditAnywhere, Category = Settings, meta = (PinHiddenByDefault))
 	bool bShouldUseCachedChannelData = false;
 	
-	// blend time over which the yaw from the animation is distributed across the trajectory samples (negative values implies yaw from the animation is constant over the entire trajectory, so the trajectory will not try to recover towards the capsule orientation)
-	UPROPERTY(EditAnywhere, Category = "Root Motion (Experimental)", meta = (PinHiddenByDefault))
-	float YawFromAnimationTrajectoryBlendTime = 0.1f;
-
-	// rate at which the root bone orientation catches up to the capsule orientation after being controlled by animation
-	// (negative values mean the capsule is authoritative over the root bone orientation and the root bone is always synchronized with the capsule,
-	// 0 means the orientation will be fully controlled by animation, and potentially never converge over the capsule orientation,
-	// positive values represent the rate at which the orientation drifts towards the capsule orientation after being controlled by animation)
-	UPROPERTY(EditAnywhere, Category = "Root Motion (Experimental)", meta = (PinHiddenByDefault))
-	float YawFromAnimationBlendRate = -1.f;
-
 	// Encapsulated motion matching algorithm and internal state
 	FMotionMatchingState MotionMatchingState;
 
diff --git a/Engine/Plugins/Animation/PoseSearch/Source/Runtime/Public/PoseSearch/PoseSearchLibrary.h b/Engine/Plugins/Animation/PoseSearch/Source/Runtime/Public/PoseSearch/PoseSearchLibrary.h
--- a/Engine/Plugins/Animation/PoseSearch/Source/Runtime/Public/PoseSearch/PoseSearchLibrary.h
+++ b/Engine/Plugins/Animation/PoseSearch/Source/Runtime/Public/PoseSearch/PoseSearchLibrary.h
@@ -56,8 +56,6 @@ struct FMotionMatchingState
 
 	void UpdateWantedPlayRate(const UE::PoseSearch::FSearchContext& SearchContext, const FFloatInterval& PlayRate, float TrajectorySpeedMultiplier);
 
-	void UpdateRootBoneControl(const FAnimationUpdateContext& Context, float YawFromAnimationBlendRate);
-
 	FVector GetEstimatedFutureRootMotionVelocity() const;
 
 	UE::PoseSearch::FSearchResult CurrentSearchResult;
@@ -152,8 +150,6 @@ public:
 		float SearchThrottleTime,
 		const FFloatInterval& PlayRate,
 		FMotionMatchingState& InOutMotionMatchingState,
-		float YawFromAnimationBlendRate = -1.f,
-		float YawFromAnimationTrajectoryBlendTime = 0.1f,
 		EPoseSearchInterruptMode InterruptMode = EPoseSearchInterruptMode::DoNotInterrupt,
 		bool bShouldSearch = true,
 		bool bShouldUseCachedChannelData = true,
diff --git a/Engine/Plugins/Experimental/Chooser/Source/Chooser/Private/AnimNode_ChooserPlayer.cpp b/Engine/Plugins/Experimental/Chooser/Source/Chooser/Private/AnimNode_ChooserPlayer.cpp
--- a/Engine/Plugins/Experimental/Chooser/Source/Chooser/Private/AnimNode_ChooserPlayer.cpp
+++ b/Engine/Plugins/Experimental/Chooser/Source/Chooser/Private/AnimNode_ChooserPlayer.cpp
@@ -159,7 +159,7 @@ void FAnimNode_ChooserPlayer::UpdateAssetPlayer(const FAnimationUpdateContext& C
 		}
 		
 		FAnimNode_BlendStack_Standalone::BlendTo(Context, NewAsset, Settings.StartTime,
-			bLoop, Settings.bMirror, MirrorDataTable, Settings.BlendTime, Settings.BlendTime,
+			bLoop, Settings.bMirror, MirrorDataTable, Settings.BlendTime,
 			Settings.BlendProfile, Settings.BlendOption, Settings.bUseInertialBlend, FVector::Zero(), Settings.PlaybackRate, 0.f,
 			GetGroupName(), GetGroupRole(), GetGroupMethod());
 
