Allow blend stack for force blend even when animation asset is null
Fixes the lack of blend when jumping

#changelist validated
#virtualized

diff --git a/Engine/Plugins/Animation/BlendStack/Source/Runtime/Private/AnimNode_BlendStack.cpp b/Engine/Plugins/Animation/BlendStack/Source/Runtime/Private/AnimNode_BlendStack.cpp
--- a/Engine/Plugins/Animation/BlendStack/Source/Runtime/Private/AnimNode_BlendStack.cpp
+++ b/Engine/Plugins/Animation/BlendStack/Source/Runtime/Private/AnimNode_BlendStack.cpp
@@ -995,7 +995,7 @@ void FAnimNode_BlendStack::UpdateAssetPlayer(const FAnimationUpdateContext& Cont
 	GetEvaluateGraphExposedInputs().Execute(Context);
 
 	bool bExecuteBlendTo = false;
-	if (AnimationAsset == nullptr)
+	if (AnimationAsset == nullptr && !bForceBlendNextUpdate)
 	{
 		bExecuteBlendTo = false;
 	}
