motion matching - set BlendStack and PoseSearch plugins production ready

diff --git a/Engine/Plugins/Animation/BlendStack/BlendStack.uplugin b/Engine/Plugins/Animation/BlendStack/BlendStack.uplugin
--- a/Engine/Plugins/Animation/BlendStack/BlendStack.uplugin
+++ b/Engine/Plugins/Animation/BlendStack/BlendStack.uplugin
@@ -1,9 +1,9 @@
 {
 	"FileVersion": 3,
 	"Version": 1,
-	"VersionName": "0.1",
+	"VersionName": "1.0",
 	"FriendlyName": "Blend Stack",
-	"Description": "Experimental blend stack API",
+	"Description": "Blend Stack API",
 	"Category": "Animation",
 	"CreatedBy": "Epic Games, Inc.",
 	"CreatedByURL": "https://epicgames.com",
@@ -12,8 +12,6 @@
 	"SupportURL": "",
 	"EnabledByDefault": false,
 	"CanContainContent": false,
-	"IsBetaVersion": false,
-	"IsExperimentalVersion": true,
 	"Installed": false,
 	"Modules": [
 		{
diff --git a/Engine/Plugins/Animation/PoseSearch/PoseSearch.uplugin b/Engine/Plugins/Animation/PoseSearch/PoseSearch.uplugin
--- a/Engine/Plugins/Animation/PoseSearch/PoseSearch.uplugin
+++ b/Engine/Plugins/Animation/PoseSearch/PoseSearch.uplugin
@@ -1,9 +1,9 @@
 {
 	"FileVersion": 3,
 	"Version": 1,
-	"VersionName": "0.1",
+	"VersionName": "1.0",
 	"FriendlyName": "Pose Search",
-	"Description": "Experimental pose search API",
+	"Description": "Pose Search API",
 	"Category": "Animation",
 	"CreatedBy": "Epic Games, Inc.",
 	"CreatedByURL": "https://epicgames.com",
@@ -12,8 +12,6 @@
 	"SupportURL": "",
 	"EnabledByDefault": false,
 	"CanContainContent": true,
-	"IsBetaVersion": false,
-	"IsExperimentalVersion": true,
 	"Installed": false,
 	"Modules": [
 		{
