Marked blend stack as not experimental, since Motion Matching production ready depends on it

diff --git a/Engine/Plugins/Animation/BlendStack/Source/Editor/Public/AnimGraphNode_BlendStack.h b/Engine/Plugins/Animation/BlendStack/Source/Editor/Public/AnimGraphNode_BlendStack.h
--- a/Engine/Plugins/Animation/BlendStack/Source/Editor/Public/AnimGraphNode_BlendStack.h
+++ b/Engine/Plugins/Animation/BlendStack/Source/Editor/Public/AnimGraphNode_BlendStack.h
@@ -8,7 +8,7 @@
 
 class UAnimGraphNode_BlendStackInput;
 
-UCLASS(Abstract, Experimental)
+UCLASS(Abstract)
 class BLENDSTACKEDITOR_API UAnimGraphNode_BlendStack_Base : public UAnimGraphNode_AssetPlayerBase
 {
 	GENERATED_BODY()
@@ -51,7 +51,7 @@ private:
 	TObjectPtr<UEdGraph> BoundGraph = nullptr;
 };
 
-UCLASS(MinimalAPI, Experimental)
+UCLASS(MinimalAPI)
 class UAnimGraphNode_BlendStack : public UAnimGraphNode_BlendStack_Base
 {
 	GENERATED_BODY()
diff --git a/Engine/Plugins/Animation/BlendStack/Source/Runtime/Public/BlendStack/AnimNode_BlendStackInput.h b/Engine/Plugins/Animation/BlendStack/Source/Runtime/Public/BlendStack/AnimNode_BlendStackInput.h
--- a/Engine/Plugins/Animation/BlendStack/Source/Runtime/Public/BlendStack/AnimNode_BlendStackInput.h
+++ b/Engine/Plugins/Animation/BlendStack/Source/Runtime/Public/BlendStack/AnimNode_BlendStackInput.h
@@ -8,7 +8,7 @@
 /** Input pose that links the blend stack's sample graph with the sample/pose chosen by the blend stack.
 *	@Todo: It might be better to reuse FAnimNode_LinkedInputPose, since we will most likely need variable input pins in the future too.
 */
-USTRUCT(BlueprintInternalUseOnly, Experimental)
+USTRUCT(BlueprintInternalUseOnly)
 struct BLENDSTACK_API FAnimNode_BlendStackInput : public FAnimNode_Base
 {
 	GENERATED_BODY()
