Move blend stack debugging to from visual logger into rewind debugger anim node property tracing

diff --git a/Engine/Plugins/Animation/BlendStack/Source/Runtime/Private/AnimNode_BlendStack.cpp b/Engine/Plugins/Animation/BlendStack/Source/Runtime/Private/AnimNode_BlendStack.cpp
--- a/Engine/Plugins/Animation/BlendStack/Source/Runtime/Private/AnimNode_BlendStack.cpp
+++ b/Engine/Plugins/Animation/BlendStack/Source/Runtime/Private/AnimNode_BlendStack.cpp
@@ -13,7 +13,6 @@
 #include "BlendStackAnimEventsFilterScope.h"
 #include "Animation/AnimInertializationSyncScope.h"
 #include "BlendStack/BlendStackDefines.h"
-#include "VisualLogger/VisualLogger.h"
 
 #include UE_INLINE_GENERATED_CPP_BY_NAME(AnimNode_BlendStack)
 
@@ -453,22 +452,32 @@ void FAnimNode_BlendStack_Standalone::Evaluate_AnyThread(FPoseContext& Output)
 #if ENABLE_ANIM_DEBUG
 	bDisableBlendStack = !CVarAnimBlendStackEnable.GetValueOnAnyThread();
 
-	const FVector& ActorLocation = Output.AnimInstanceProxy->GetActorTransform().GetTranslation();
-	
-	TStringBuilder<1024> MessageBuilder;
-
-	for (int32 i = 0; i < AnimPlayers.Num(); ++i)
+	if (UE_TRACE_CHANNELEXPR_IS_ENABLED(AnimationChannel))
 	{
-		const FBlendStackAnimPlayer& AnimPlayer = AnimPlayers[i];
-		MessageBuilder.Appendf(TEXT("%d) t:%.2f/%.2f a:%.2f m:%d %s\n"), 
-			i, AnimPlayer.GetCurrentBlendInTime(), AnimPlayer.GetTotalBlendInTime(), AnimPlayer.GetTimeToActivation(),
-			AnimPlayer.GetMirror() ? 1 : 0, *AnimPlayer.GetAnimationName());
-	}
+		// output current asset as "Asset" because that column is shown by default
+		TRACE_ANIM_NODE_VALUE(Output, TEXT("Asset"), GetAnimAsset());
 
-	FString VLogMessage(MessageBuilder.ToString());
-	FString VLogCategoryName;
-	VLogCategoryName.Appendf(TEXT("LogBlendStack_%d"), GetNodeIndex());
-	UE_VLOG_LOCATION(Output.GetAnimInstanceObject(), *VLogCategoryName, Verbose, ActorLocation + FVector(0, 0, 100.0f), 0.f, FColor::Yellow, TEXT("%s"), *VLogMessage);
+		for (int32 i = 0; i < AnimPlayers.Num(); ++i)
+		{
+			if (AnimPlayers[i].IsActive())
+			{
+				FString IndexString = FString("[" + FString::FromInt(i) + FString("]"));
+				FString Asset = FString("Asset") + IndexString;
+				FString ElapsedTime = FString("ElapsedTime") + IndexString;
+				FString CurrentBlendInTime = FString("CurrentBlendInTime") + IndexString;
+				FString TotalBlendInTime = FString("TotalBlendInTime") + IndexString;
+				FString TimeToActivation = FString("TimeToActivation") + IndexString;
+				FString Mirror = FString("Mirror") + IndexString;
+
+				const FBlendStackAnimPlayer& AnimPlayer = AnimPlayers[i];
+				TRACE_ANIM_NODE_VALUE(Output, *Asset, AnimPlayer.GetAnimationAsset());
+				TRACE_ANIM_NODE_VALUE(Output, *ElapsedTime, AnimPlayer.GetAccumulatedTime());
+				TRACE_ANIM_NODE_VALUE(Output, *CurrentBlendInTime, AnimPlayer.GetCurrentBlendInTime());
+				TRACE_ANIM_NODE_VALUE(Output, *TotalBlendInTime, AnimPlayer.GetTotalBlendInTime());
+				TRACE_ANIM_NODE_VALUE(Output, *TimeToActivation, AnimPlayer.GetTimeToActivation());
+			}
+		}
+	}
 #endif // ENABLE_ANIM_DEBUG
 
 
@@ -735,7 +744,7 @@ void FAnimNode_BlendStack_Standalone::UpdateSample(const FAnimationUpdateContext
 			const int32 SampleIndex = SamplePlayer.GetPoseLinkIndex();
 			UE::BlendStack::FBlendStack_SampleGraphExecutionHelper& ExecutionHelper = SampleGraphExecutionHelpers[SampleIndex];
 			ExecutionHelper.SetInputPosePlayer(SamplePlayer);
-			// The anim player may or may not have its Update_AnyThread called through the graph update. 
+			// The anim player may or may not have its Update_AnyThread called through the graph update.
 			PerSampleGraphPoseLinks[SampleIndex].Update(Context);
 		}
 		else
diff --git a/Engine/Plugins/Experimental/Chooser/Source/Chooser/Private/AnimNode_ChooserPlayer.cpp b/Engine/Plugins/Experimental/Chooser/Source/Chooser/Private/AnimNode_ChooserPlayer.cpp
--- a/Engine/Plugins/Experimental/Chooser/Source/Chooser/Private/AnimNode_ChooserPlayer.cpp
+++ b/Engine/Plugins/Experimental/Chooser/Source/Chooser/Private/AnimNode_ChooserPlayer.cpp
@@ -209,7 +209,6 @@ void FAnimNode_ChooserPlayer::UpdateAssetPlayer(const FAnimationUpdateContext& C
 	Chooser.Get<FObjectChooserBase>().GetDebugName(Name);
 	
 	TRACE_ANIM_NODE_VALUE(Context, TEXT("Name"), ToCStr(Name));
-	TRACE_ANIM_NODE_VALUE(Context, TEXT("Asset"), CurrentAsset);
 }
 
 void FAnimNode_ChooserPlayer::Evaluate_AnyThread(FPoseContext& Output)
