blendstack - stitch improvements

diff --git a/Engine/Plugins/Animation/BlendStack/Source/Runtime/Private/AnimNode_BlendStack.cpp b/Engine/Plugins/Animation/BlendStack/Source/Runtime/Private/AnimNode_BlendStack.cpp
--- a/Engine/Plugins/Animation/BlendStack/Source/Runtime/Private/AnimNode_BlendStack.cpp
+++ b/Engine/Plugins/Animation/BlendStack/Source/Runtime/Private/AnimNode_BlendStack.cpp
@@ -832,51 +832,58 @@ void FAnimNode_BlendStack_Standalone::BlendTo(const FAnimationUpdateContext& Con
 	bool bNeedToBlendTo = true;
 	if (StitchDatabase)
 	{
-		if (IPoseSearchProvider* PoseSearchProvider = IPoseSearchProvider::Get())
+		if (MaxActiveBlends > 0)
 		{
-			// looking for an animation stitch from the StitchDatabase that will connect, in BlendTime seconds, 
-			// the currently playing animation pose to the pose from AnimationAsset at AccumulatedTime + BlendTime
-			const UObject* AssetToSearch = StitchDatabase.Get();
+			if (IPoseSearchProvider* PoseSearchProvider = IPoseSearchProvider::Get())
+			{
+				// looking for an animation stitch from the StitchDatabase that will connect, in BlendTime seconds, 
+				// the currently playing animation pose to the pose from AnimationAsset at AccumulatedTime + BlendTime
+				const UObject* AssetToSearch = StitchDatabase.Get();
 
-			IPoseSearchProvider::FSearchPlayingAsset PlayingAsset;
-			PlayingAsset.Asset = GetAnimAsset();
-			PlayingAsset.AccumulatedTime = GetAccumulatedTime();
+				IPoseSearchProvider::FSearchPlayingAsset PlayingAsset;
+				PlayingAsset.Asset = GetAnimAsset();
+				PlayingAsset.AccumulatedTime = GetAccumulatedTime();
 
-			IPoseSearchProvider::FSearchFutureAsset FutureAsset;
-			FutureAsset.Asset = AnimationAsset;
-			FutureAsset.AccumulatedTime = AccumulatedTime + BlendTime;
-			FutureAsset.IntervalTime = BlendTime;
+				IPoseSearchProvider::FSearchFutureAsset FutureAsset;
+				FutureAsset.Asset = AnimationAsset;
+				FutureAsset.AccumulatedTime = AccumulatedTime + BlendTime;
+				FutureAsset.IntervalTime = BlendTime;
 
-			const IPoseSearchProvider::FSearchResult SearchResult = PoseSearchProvider->Search(Context, MakeArrayView(&AssetToSearch, 1), PlayingAsset, FutureAsset);
-			if (UAnimationAsset* StitchAnimationAsset = Cast<UAnimationAsset>(SearchResult.SelectedAsset))
-			{
-				if (SearchResult.Dissimilarity <= StitchBlendMaxCost)
+				const IPoseSearchProvider::FSearchResult SearchResult = PoseSearchProvider->Search(Context, MakeArrayView(&AssetToSearch, 1), PlayingAsset, FutureAsset);
+				if (UAnimationAsset* StitchAnimationAsset = Cast<UAnimationAsset>(SearchResult.SelectedAsset))
 				{
-					// blend to the selected animation stitch
-					InternalBlendTo(Context, StitchAnimationAsset, SearchResult.TimeOffsetSeconds, false, SearchResult.bMirrored, MirrorDataTable,
-						StitchBlendTime, BlendProfile, BlendOption, bUseInertialBlend, BlendParameters, SearchResult.WantedPlayRate, ActivationDelay, GroupName, GroupRole, GroupMethod);
-
-					// blend with an ActivationDelay of BlendTime - StitchBlendTime + ActivationDelay seconds
-					// to the AnimationAsset at AccumulatedTime + BlendTime - StitchBlendTime seconds in the future,
-					// so at BlendTime seconds ahead the AnimationAsset is playing the fully blended in pose at AccumulatedTime + BlendTime
-					InternalBlendTo(Context, AnimationAsset, AccumulatedTime + BlendTime - StitchBlendTime, bLoop, bMirrored, MirrorDataTable,
-						StitchBlendTime, BlendProfile, BlendOption, bUseInertialBlend, BlendParameters, PlayRate, BlendTime - StitchBlendTime + ActivationDelay, GroupName, GroupRole, GroupMethod);
-
-					bNeedToBlendTo = false;
+					if (SearchResult.Dissimilarity <= StitchBlendMaxCost)
+					{
+						// blend to the selected animation stitch
+						InternalBlendTo(Context, StitchAnimationAsset, SearchResult.TimeOffsetSeconds, false, SearchResult.bMirrored, MirrorDataTable,
+							StitchBlendTime, BlendProfile, BlendOption, bUseInertialBlend, BlendParameters, SearchResult.WantedPlayRate, ActivationDelay, GroupName, GroupRole, GroupMethod);
+
+						// blend with an ActivationDelay of BlendTime - StitchBlendTime + ActivationDelay seconds
+						// to the AnimationAsset at AccumulatedTime + BlendTime - StitchBlendTime seconds in the future,
+						// so at BlendTime seconds ahead the AnimationAsset is playing the fully blended in pose at AccumulatedTime + BlendTime
+						InternalBlendTo(Context, AnimationAsset, AccumulatedTime + BlendTime - StitchBlendTime, bLoop, bMirrored, MirrorDataTable,
+							StitchBlendTime, BlendProfile, BlendOption, bUseInertialBlend, BlendParameters, PlayRate, BlendTime - StitchBlendTime + ActivationDelay, GroupName, GroupRole, GroupMethod);
+
+						bNeedToBlendTo = false;
+					}
+					else
+					{
+						UE_LOG(LogBlendStack, Display, TEXT("FAnimNode_BlendStack_Standalone::BlendTo StitchDatabase '%s' search cost is %f, above StitchBlendMaxCost %f. Defaulting to regular blend"), *GetNameSafe(StitchDatabase), SearchResult.Dissimilarity, StitchBlendMaxCost);
+					}
 				}
 				else
 				{
-					UE_LOG(LogBlendStack, Display, TEXT("FAnimNode_BlendStack_Standalone::BlendTo StitchDatabase '%s' search cost is %f, above StitchBlendMaxCost %f. Defaulting to regular blend"), *GetNameSafe(StitchDatabase), SearchResult.Dissimilarity, StitchBlendMaxCost);
+					UE_LOG(LogBlendStack, Error, TEXT("FAnimNode_BlendStack_Standalone::BlendTo cannot use StitchDatabase '%s', because of missing IPoseSearchProvider::Search couldn't select a StitchAnimationAsset. Defaulting to regular blend"), *GetNameSafe(StitchDatabase));
 				}
 			}
 			else
 			{
-				UE_LOG(LogBlendStack, Error, TEXT("FAnimNode_BlendStack_Standalone::BlendTo cannot use StitchDatabase '%s', because of missing IPoseSearchProvider::Search couldn't select a StitchAnimationAsset. Defaulting to regular blend"), *GetNameSafe(StitchDatabase));
+				UE_LOG(LogBlendStack, Error, TEXT("FAnimNode_BlendStack_Standalone::BlendTo cannot use StitchDatabase '%s', because of missing IPoseSearchProvider (is PoseSearch plugin enabled?). Defaulting to regular blend"), *GetNameSafe(StitchDatabase));
 			}
 		}
 		else
 		{
-			UE_LOG(LogBlendStack, Error, TEXT("FAnimNode_BlendStack_Standalone::BlendTo cannot use StitchDatabase '%s', because of missing IPoseSearchProvider (is PoseSearch plugin enabled?). Defaulting to regular blend"), *GetNameSafe(StitchDatabase));
+			UE_LOG(LogBlendStack, Error, TEXT("FAnimNode_BlendStack_Standalone::BlendTo cannot use StitchDatabase '%s', since MaxActiveBlends should be at least 1. Defaulting to regular blend"), *GetNameSafe(StitchDatabase));
 		}
 	}
 
