blend stack - implementing UAnimGraphNode_BlendStack::BakeDataDuringCompilation in UAnimGraphNode_BlendStack_Base to generalize AnimBlueprint->FindOrAddGroup(...)

diff --git a/Engine/Plugins/Animation/BlendStack/Source/Editor/Private/AnimGraphNode_BlendStack.cpp b/Engine/Plugins/Animation/BlendStack/Source/Editor/Private/AnimGraphNode_BlendStack.cpp
--- a/Engine/Plugins/Animation/BlendStack/Source/Editor/Private/AnimGraphNode_BlendStack.cpp
+++ b/Engine/Plugins/Animation/BlendStack/Source/Editor/Private/AnimGraphNode_BlendStack.cpp
@@ -43,10 +43,12 @@ void UAnimGraphNode_BlendStack_Base::GetOutputLinkAttributes(FNodeAttributeArray
 	}
 }
 
-void UAnimGraphNode_BlendStack::BakeDataDuringCompilation(class FCompilerResultsLog& MessageLog)
+void UAnimGraphNode_BlendStack_Base::BakeDataDuringCompilation(class FCompilerResultsLog& MessageLog)
 {
+	const FAnimNode_BlendStack_Standalone* BlendStackNode = GetBlendStackNode();
+	check(BlendStackNode);
 	UAnimBlueprint* AnimBlueprint = GetAnimBlueprint();
-	AnimBlueprint->FindOrAddGroup(Node.GetGroupName());
+	AnimBlueprint->FindOrAddGroup(BlendStackNode->GetGroupName());
 }
 
 void UAnimGraphNode_BlendStack_Base::ExpandGraphAndProcessNodes(
@@ -155,7 +157,7 @@ void UAnimGraphNode_BlendStack_Base::OnProcessDuringCompilation(IAnimBlueprintCo
 	const int32 MaxBlendsNum = GetMaxActiveBlends();
 	// Set MaxActiveBlends of the blend stack at compile time.
 	// @todo: Allow reducing the blend stack size (but not decreasing), to enable scalability control under i.e. lower LODs
-	AnimNode->MaxActiveBlends = GetMaxActiveBlends();
+	AnimNode->SetMaxActiveBlends(MaxBlendsNum);
 
 	UAnimationBlendStackGraph* SampleGraph = CastChecked<UAnimationBlendStackGraph>(BoundGraph);
 	if (!SampleGraph->ResultNode->IsNodeRootSet())
@@ -234,7 +236,7 @@ int32 UAnimGraphNode_BlendStack_Base::GetMaxActiveBlends() const
 {
 	FAnimNode_BlendStack_Standalone* AnimNode = GetBlendStackNode();
 	check(AnimNode);
-	return AnimNode->MaxActiveBlends;
+	return AnimNode->GetMaxActiveBlends();
 }
 
 void UAnimGraphNode_BlendStack_Base::CreateGraph()
diff --git a/Engine/Plugins/Animation/BlendStack/Source/Runtime/Public/BlendStack/AnimNode_BlendStack.h b/Engine/Plugins/Animation/BlendStack/Source/Runtime/Public/BlendStack/AnimNode_BlendStack.h
--- a/Engine/Plugins/Animation/BlendStack/Source/Runtime/Public/BlendStack/AnimNode_BlendStack.h
+++ b/Engine/Plugins/Animation/BlendStack/Source/Runtime/Public/BlendStack/AnimNode_BlendStack.h
@@ -156,7 +156,8 @@ struct BLENDSTACK_API FAnimNode_BlendStack_Standalone : public FAnimNode_AssetPl
 	int32 GetNextPoseLinkIndex();
 
 	int32 GetMaxActiveBlends() const { return MaxActiveBlends; }
-
+	void SetMaxActiveBlends(int32 InMaxActiveBlends) { MaxActiveBlends = InMaxActiveBlends; }
+	
 protected:
 	void InternalBlendTo(const FAnimationUpdateContext& Context, UAnimationAsset* AnimationAsset, float AccumulatedTime, bool bLoop, 
 		bool bMirrored, UMirrorDataTable* MirrorDataTable, float BlendTime,
@@ -200,18 +201,15 @@ private:
 	void UpdateSample(const FAnimationUpdateContext& Context, const int32 PlayerIndex);
 	void CacheBonesForSample(const FAnimationCacheBonesContext& Context, const int32 PlayerIndex);
 	bool IsSampleGraphAvailableForPlayer(const int32 PlayerIndex);
-
-	friend class UAnimGraphNode_BlendStack_Base;
 };
 
-
 USTRUCT(BlueprintInternalUseOnly)
 struct BLENDSTACK_API FAnimNode_BlendStack : public FAnimNode_BlendStack_Standalone
 {
 	GENERATED_BODY()
 
 	// requested animation to play
-	UPROPERTY(EditAnywhere, BlueprintReadWrite, Category = Settings, meta = (PinShownByDefault))
+	UPROPERTY(EditAnywhere, BlueprintReadWrite, Category = Settings, meta = (PinHiddenByDefault))
 	TObjectPtr<UAnimationAsset> AnimationAsset;
 
 	// requested animation time
@@ -266,9 +264,6 @@ struct BLENDSTACK_API FAnimNode_BlendStack : public FAnimNode_BlendStack_Standal
 	UPROPERTY(EditAnywhere, Category = Settings)
 	bool bResetOnBecomingRelevant = true;
 
-	// Update Counter for detecting being relevant
-	FGraphTraversalCounter UpdateCounter;
-
 	// FAnimNode_Base interface
 	virtual void UpdateAssetPlayer(const FAnimationUpdateContext& Context) override;
 	// End of FAnimNode_Base interface
@@ -277,8 +272,9 @@ struct BLENDSTACK_API FAnimNode_BlendStack : public FAnimNode_BlendStack_Standal
 	void ForceBlendNextUpdate();
 
 private:
+	// Update Counter for detecting being relevant
+	FGraphTraversalCounter UpdateCounter;
 
 	bool bForceBlendNextUpdate = false;
-
 };
 
