Add SyncGroup support to AnimNode_BlendStack

diff --git a/Engine/Plugins/Animation/BlendStack/Source/Runtime/Private/AnimNode_BlendStack.cpp b/Engine/Plugins/Animation/BlendStack/Source/Runtime/Private/AnimNode_BlendStack.cpp
--- a/Engine/Plugins/Animation/BlendStack/Source/Runtime/Private/AnimNode_BlendStack.cpp
+++ b/Engine/Plugins/Animation/BlendStack/Source/Runtime/Private/AnimNode_BlendStack.cpp
@@ -1139,7 +1139,8 @@ void FAnimNode_BlendStack::UpdateAssetPlayer(const FAnimationUpdateContext& Cont
 	if (bExecuteBlendTo)
 	{
 		BlendTo(Context, AnimationAsset, AnimationTime, bLoop, bMirrored, MirrorDataTable.Get(),
-			BlendTime, BlendProfile, BlendOption, bUseInertialBlend, BlendParameters, WantedPlayRate, ActivationDelayTime);
+			BlendTime, BlendProfile, BlendOption, bUseInertialBlend, BlendParameters, WantedPlayRate, ActivationDelayTime,
+			GetGroupName(), GetGroupRole(), GetGroupMethod());
 	}
 
 	const bool bDidBlendToRequestAnInertialBlend = bExecuteBlendTo && bUseInertialBlend;
@@ -1161,4 +1162,94 @@ void UE::BlendStack::FBlendStack_SampleGraphExecutionHelper::SetInputPosePlayer(
 	Player = &InPlayer;
 }
 
+FName FAnimNode_BlendStack::GetGroupName() const
+{
+	return GET_ANIM_NODE_DATA(FName, GroupName);
+}
+
+EAnimGroupRole::Type FAnimNode_BlendStack::GetGroupRole() const
+{
+	return GET_ANIM_NODE_DATA(TEnumAsByte<EAnimGroupRole::Type>, GroupRole);
+}
+
+EAnimSyncMethod FAnimNode_BlendStack::GetGroupMethod() const
+{
+	return GET_ANIM_NODE_DATA(EAnimSyncMethod, Method);
+}
+
+bool FAnimNode_BlendStack::GetIgnoreForRelevancyTest() const
+{
+	return GET_ANIM_NODE_DATA(bool, bIgnoreForRelevancyTest);
+}
+
+bool FAnimNode_BlendStack::SetGroupName(FName InGroupName)
+{
+#if WITH_EDITORONLY_DATA
+	GroupName = InGroupName;
+#endif
+
+	if(FName* GroupNamePtr = GET_INSTANCE_ANIM_NODE_DATA_PTR(FName, GroupName))
+	{
+		*GroupNamePtr = InGroupName;
+		return true;
+	}
+
+	return false;
+}
+
+bool FAnimNode_BlendStack::SetGroupRole(EAnimGroupRole::Type InRole)
+{
+#if WITH_EDITORONLY_DATA
+	GroupRole = InRole;
+#endif
+	
+	if(TEnumAsByte<EAnimGroupRole::Type>* GroupRolePtr = GET_INSTANCE_ANIM_NODE_DATA_PTR(TEnumAsByte<EAnimGroupRole::Type>, GroupRole))
+	{
+		*GroupRolePtr = InRole;
+		return true;
+	}
+
+	return false;
+}
+
+bool FAnimNode_BlendStack::SetGroupMethod(EAnimSyncMethod InMethod)
+{
+#if WITH_EDITORONLY_DATA
+	Method = InMethod;
+#endif
+
+	if(EAnimSyncMethod* MethodPtr = GET_INSTANCE_ANIM_NODE_DATA_PTR(EAnimSyncMethod, Method))
+	{
+		*MethodPtr = InMethod;
+		return true;
+	}
+
+	return false;
+}
+
+bool FAnimNode_BlendStack::IsLooping() const
+{
+	if (!AnimPlayers.IsEmpty())
+	{
+		return AnimPlayers[0].IsLooping();
+	}
+	return false;
+}
+
+
+bool FAnimNode_BlendStack::SetIgnoreForRelevancyTest(bool bInIgnoreForRelevancyTest)
+{
+#if WITH_EDITORONLY_DATA
+	bIgnoreForRelevancyTest = bInIgnoreForRelevancyTest;
+#endif
+
+	if (bool* bIgnoreForRelevancyTestPtr = GET_INSTANCE_ANIM_NODE_DATA_PTR(bool, bIgnoreForRelevancyTest))
+	{
+		*bIgnoreForRelevancyTestPtr = bInIgnoreForRelevancyTest;
+		return true;
+	}
+
+	return false;
+}
+
 #undef LOCTEXT_NAMESPACE
diff --git a/Engine/Plugins/Animation/BlendStack/Source/Runtime/Private/BlendStackAnimNodeLibrary.cpp b/Engine/Plugins/Animation/BlendStack/Source/Runtime/Private/BlendStackAnimNodeLibrary.cpp
--- a/Engine/Plugins/Animation/BlendStack/Source/Runtime/Private/BlendStackAnimNodeLibrary.cpp
+++ b/Engine/Plugins/Animation/BlendStack/Source/Runtime/Private/BlendStackAnimNodeLibrary.cpp
@@ -64,7 +64,8 @@ void UBlendStackAnimNodeLibrary::BlendTo(const FAnimUpdateContext& Context,
 					BlendStackNodePtr->bUseInertialBlend,
 					BlendParameters,
 					WantedPlayRate,
-					ActivationDelay);
+					ActivationDelay,
+					BlendStackNodePtr->GetGroupName(), BlendStackNodePtr->GetGroupRole(), BlendStackNodePtr->GetGroupMethod());
 			}
 			else
 			{
diff --git a/Engine/Plugins/Animation/BlendStack/Source/Runtime/Public/BlendStack/AnimNode_BlendStack.h b/Engine/Plugins/Animation/BlendStack/Source/Runtime/Public/BlendStack/AnimNode_BlendStack.h
--- a/Engine/Plugins/Animation/BlendStack/Source/Runtime/Public/BlendStack/AnimNode_BlendStack.h
+++ b/Engine/Plugins/Animation/BlendStack/Source/Runtime/Public/BlendStack/AnimNode_BlendStack.h
@@ -264,8 +264,39 @@ struct BLENDSTACK_API FAnimNode_BlendStack : public FAnimNode_BlendStack_Standal
 	UPROPERTY(EditAnywhere, Category = Settings)
 	bool bResetOnBecomingRelevant = true;
 
+	#if WITH_EDITORONLY_DATA
+    	// The group name that we synchronize with (NAME_None if it is not part of any group). Note that
+    	// this is the name of the group used to sync the output of this node - it will not force
+    	// syncing of animations contained by it.
+    	UPROPERTY(EditAnywhere, Category = Sync, meta = (FoldProperty))
+    	FName GroupName = NAME_None;
+    
+    	// The role this node can assume within the group (ignored if GroupName is not set). Note
+    	// that this is the role of the output of this node, not of animations contained by it.
+    	UPROPERTY(EditAnywhere, Category = Sync, meta = (FoldProperty))
+    	TEnumAsByte<EAnimGroupRole::Type> GroupRole = EAnimGroupRole::CanBeLeader;
+    
+    	// How this node will synchronize with other animations. Note that this determines how the output
+    	// of this node is used for synchronization, not of animations contained by it.
+    	UPROPERTY(EditAnywhere, Category = Sync, meta = (FoldProperty))
+    	EAnimSyncMethod Method = EAnimSyncMethod::DoNotSync;
+    
+    	// If true, "Relevant anim" nodes that look for the highest weighted animation in a state will ignore this node
+    	UPROPERTY(EditAnywhere, Category = Relevancy, meta = (FoldProperty, PinHiddenByDefault))
+    	bool bIgnoreForRelevancyTest = false;
+    #endif // WITH_EDITORONLY_DATA
+
 	// FAnimNode_Base interface
 	virtual void UpdateAssetPlayer(const FAnimationUpdateContext& Context) override;
+	virtual FName GetGroupName() const override;
+	virtual EAnimGroupRole::Type GetGroupRole() const override;
+	virtual EAnimSyncMethod GetGroupMethod() const override;
+	virtual bool GetIgnoreForRelevancyTest() const override;
+	virtual bool IsLooping() const override;
+	virtual bool SetGroupName(FName InGroupName) override;
+	virtual bool SetGroupRole(EAnimGroupRole::Type InRole) override;
+	virtual bool SetGroupMethod(EAnimSyncMethod InMethod) override;
+	virtual bool SetIgnoreForRelevancyTest(bool bInIgnoreForRelevancyTest) override;
 	// End of FAnimNode_Base interface
 
 	// Force a blend on the next update, even if the anim sequence has not changed.
