[Motion Matching] - Added sync group support.

Allows Motion Matching node to be part of a named sync group as an leader, no follower roles are allowed since that would break the pose matching already done. Lastly, only the most recent sample in the blend stack has sync marker information and is the leader for other asset players to follow, the rest of the other blend stack samples do not sync.

Added flag to stomp previous leader info (bOverwriteLeaderPosition) and a new role (ExclusiveLeader) which will make a tick record tick as ungrouped if it fails to be a leader.

[RN] minor

diff --git a/Engine/Plugins/Animation/BlendStack/Source/Runtime/Private/AnimNode_BlendStack.cpp b/Engine/Plugins/Animation/BlendStack/Source/Runtime/Private/AnimNode_BlendStack.cpp
--- a/Engine/Plugins/Animation/BlendStack/Source/Runtime/Private/AnimNode_BlendStack.cpp
+++ b/Engine/Plugins/Animation/BlendStack/Source/Runtime/Private/AnimNode_BlendStack.cpp
@@ -28,7 +28,7 @@ static constexpr FBoneIndexType RootBoneIndexType = 0;
 // FBlendStackAnimPlayer
 void FBlendStackAnimPlayer::Initialize(const FAnimationInitializeContext& Context, UAnimationAsset* AnimationAsset, float AccumulatedTime, bool bLoop,
 	bool bMirrored, UMirrorDataTable* MirrorDataTable, float BlendTime, const UBlendProfile* BlendProfile, EAlphaBlendOption InBlendOption,
-	const FVector& BlendParameters, float PlayRate, float ActivationDelay, int32 InPoseLinkIdx, FName GroupName, EAnimGroupRole::Type GroupRole, EAnimSyncMethod GroupMethod)
+	const FVector& BlendParameters, float PlayRate, float ActivationDelay, int32 InPoseLinkIdx, FName GroupName, EAnimGroupRole::Type GroupRole, EAnimSyncMethod GroupMethod, bool bOverridePositionWhenJoiningSyncGroupAsLeader)
 {
 	if (bMirrored && !MirrorDataTable)
 	{
@@ -79,6 +79,7 @@ void FBlendStackAnimPlayer::Initialize(const FAnimationInitializeContext& Contex
 		SequencePlayerNode.SetGroupMethod(GroupMethod);
 		SequencePlayerNode.SetGroupName(GroupName);
 		SequencePlayerNode.SetGroupRole(GroupRole);
+		SequencePlayerNode.SetOverridePositionWhenJoiningSyncGroupAsLeader(bOverridePositionWhenJoiningSyncGroupAsLeader);
 	}
 	else if (UBlendSpace* BlendSpace = Cast<UBlendSpace>(AnimationAsset))
 	{
@@ -94,6 +95,7 @@ void FBlendStackAnimPlayer::Initialize(const FAnimationInitializeContext& Contex
 		BlendSpacePlayerNode.SetGroupMethod(GroupMethod);
 		BlendSpacePlayerNode.SetGroupName(GroupName);
 		BlendSpacePlayerNode.SetGroupRole(GroupRole);
+		BlendSpacePlayerNode.SetOverridePositionWhenJoiningSyncGroupAsLeader(bOverridePositionWhenJoiningSyncGroupAsLeader);
 	}
 	else if (AnimationAsset)
 	{
@@ -300,6 +302,21 @@ bool FBlendStackAnimPlayer::IsActive() const
 	return TimeToActivation <= 0.f;
 }
 
+FAnimNode_AssetPlayerBase* FBlendStackAnimPlayer::GetAssetPlayerNode()
+{
+	if (SequencePlayerNode.GetSequence())
+	{
+		return &SequencePlayerNode;
+	}
+	else if (BlendSpacePlayerNode.GetBlendSpace())
+	{
+		return &BlendSpacePlayerNode;
+	}
+
+	// Anim player was initialized with an unsupported asset type.
+	return nullptr;
+}
+
 void FBlendStackAnimPlayer::UpdateWithDeltaTime(float DeltaTime, int32 PlayerDepth, float PlayerDepthBlendInTimeMultiplier)
 {
 	const bool bIsMainPlayer = PlayerDepth == 0;
@@ -825,7 +842,7 @@ static void RequestInertialBlend(const FAnimationUpdateContext& Context, float B
 
 void FAnimNode_BlendStack_Standalone::BlendTo(const FAnimationUpdateContext& Context, UAnimationAsset* AnimationAsset, float AccumulatedTime, bool bLoop,
 	bool bMirrored, UMirrorDataTable* MirrorDataTable, float BlendTime, const UBlendProfile* BlendProfile, EAlphaBlendOption BlendOption, bool bUseInertialBlend,
-	const FVector& BlendParameters, float PlayRate, float ActivationDelay, FName GroupName, EAnimGroupRole::Type GroupRole, EAnimSyncMethod GroupMethod)
+	const FVector& BlendParameters, float PlayRate, float ActivationDelay, FName GroupName, EAnimGroupRole::Type GroupRole, EAnimSyncMethod GroupMethod, bool bOverridePositionWhenJoiningSyncGroupAsLeader)
 {
 	using namespace UE::Anim;
 
@@ -856,13 +873,13 @@ void FAnimNode_BlendStack_Standalone::BlendTo(const FAnimationUpdateContext& Con
 					{
 						// blend to the selected animation stitch
 						InternalBlendTo(Context, StitchAnimationAsset, SearchResult.TimeOffsetSeconds, false, SearchResult.bMirrored, MirrorDataTable,
-							StitchBlendTime, BlendProfile, BlendOption, bUseInertialBlend, BlendParameters, SearchResult.WantedPlayRate, ActivationDelay, GroupName, GroupRole, GroupMethod);
+							StitchBlendTime, BlendProfile, BlendOption, bUseInertialBlend, BlendParameters, SearchResult.WantedPlayRate, ActivationDelay, GroupName, GroupRole, GroupMethod, bOverridePositionWhenJoiningSyncGroupAsLeader);
 
 						// blend with an ActivationDelay of BlendTime - StitchBlendTime + ActivationDelay seconds
 						// to the AnimationAsset at AccumulatedTime + BlendTime - StitchBlendTime seconds in the future,
 						// so at BlendTime seconds ahead the AnimationAsset is playing the fully blended in pose at AccumulatedTime + BlendTime
 						InternalBlendTo(Context, AnimationAsset, AccumulatedTime + BlendTime - StitchBlendTime, bLoop, bMirrored, MirrorDataTable,
-							StitchBlendTime, BlendProfile, BlendOption, bUseInertialBlend, BlendParameters, PlayRate, BlendTime - StitchBlendTime + ActivationDelay, GroupName, GroupRole, GroupMethod);
+							StitchBlendTime, BlendProfile, BlendOption, bUseInertialBlend, BlendParameters, PlayRate, BlendTime - StitchBlendTime + ActivationDelay, GroupName, GroupRole, GroupMethod, bOverridePositionWhenJoiningSyncGroupAsLeader);
 
 						bNeedToBlendTo = false;
 					}
@@ -890,13 +907,13 @@ void FAnimNode_BlendStack_Standalone::BlendTo(const FAnimationUpdateContext& Con
 	if (bNeedToBlendTo)
 	{
 		InternalBlendTo(Context, AnimationAsset, AccumulatedTime, bLoop, bMirrored, MirrorDataTable,
-			BlendTime, BlendProfile, BlendOption, bUseInertialBlend, BlendParameters, PlayRate, ActivationDelay, GroupName, GroupRole, GroupMethod);
+			BlendTime, BlendProfile, BlendOption, bUseInertialBlend, BlendParameters, PlayRate, ActivationDelay, GroupName, GroupRole, GroupMethod, bOverridePositionWhenJoiningSyncGroupAsLeader);
 	}
 }
 
 void FAnimNode_BlendStack_Standalone::InternalBlendTo(const FAnimationUpdateContext& Context, UAnimationAsset* AnimationAsset, float AccumulatedTime, bool bLoop,
 	bool bMirrored, UMirrorDataTable* MirrorDataTable, float BlendTime, const UBlendProfile* BlendProfile, EAlphaBlendOption BlendOption, bool bUseInertialBlend,
-	const FVector& BlendParameters, float PlayRate, float ActivationDelay, FName GroupName, EAnimGroupRole::Type GroupRole, EAnimSyncMethod GroupMethod)
+	const FVector& BlendParameters, float PlayRate, float ActivationDelay, FName GroupName, EAnimGroupRole::Type GroupRole, EAnimSyncMethod GroupMethod, bool bOverridePositionWhenJoiningSyncGroupAsLeader)
 {
 	const bool bBlendStackIsEmpty = AnimPlayers.IsEmpty();
 
@@ -940,7 +957,7 @@ void FAnimNode_BlendStack_Standalone::InternalBlendTo(const FAnimationUpdateCont
 	FBlendStackAnimPlayer& AnimPlayer = AnimPlayers[0];
 
 	FAnimationInitializeContext InitContext(Context.AnimInstanceProxy, Context.SharedContext);
-	AnimPlayer.Initialize(InitContext, AnimationAsset, AccumulatedTime, bLoop, bMirrored, MirrorDataTable, BlendTime, BlendProfile, BlendOption, BlendParameters, PlayRate, ActivationDelay, NewSamplePoseLinkIndex, GroupName, GroupRole, GroupMethod);
+	AnimPlayer.Initialize(InitContext, AnimationAsset, AccumulatedTime, bLoop, bMirrored, MirrorDataTable, BlendTime, BlendProfile, BlendOption, BlendParameters, PlayRate, ActivationDelay, NewSamplePoseLinkIndex, GroupName, GroupRole, GroupMethod, bOverridePositionWhenJoiningSyncGroupAsLeader);
 	InitializeSample(InitContext, AnimPlayer);
 }
 
diff --git a/Engine/Plugins/Animation/BlendStack/Source/Runtime/Public/BlendStack/AnimNode_BlendStack.h b/Engine/Plugins/Animation/BlendStack/Source/Runtime/Public/BlendStack/AnimNode_BlendStack.h
--- a/Engine/Plugins/Animation/BlendStack/Source/Runtime/Public/BlendStack/AnimNode_BlendStack.h
+++ b/Engine/Plugins/Animation/BlendStack/Source/Runtime/Public/BlendStack/AnimNode_BlendStack.h
@@ -15,7 +15,7 @@ struct BLENDSTACK_API FBlendStackAnimPlayer
 	
 	void Initialize(const FAnimationInitializeContext& Context, UAnimationAsset* AnimationAsset, float AccumulatedTime, bool bLoop,
 		bool bMirrored, UMirrorDataTable* MirrorDataTable, float BlendTime, const UBlendProfile* BlendProfile, EAlphaBlendOption InBlendOption,
-		const FVector& BlendParameters, float PlayRate, float ActivationDelay, int32 InPoseLinkIdx, FName GroupName, EAnimGroupRole::Type GroupRole, EAnimSyncMethod Method);
+		const FVector& BlendParameters, float PlayRate, float ActivationDelay, int32 InPoseLinkIdx, FName GroupName, EAnimGroupRole::Type GroupRole, EAnimSyncMethod Method, bool bOverridePositionWhenJoiningSyncGroupAsLeader);
 	
 	void UpdatePlayRate(float PlayRate);
 	void Evaluate_AnyThread(FPoseContext& Output);
@@ -52,6 +52,8 @@ struct BLENDSTACK_API FBlendStackAnimPlayer
 	bool IsLooping() const;
 	bool IsActive() const;
 
+	FAnimNode_AssetPlayerBase* GetAssetPlayerNode();
+
 	// Curves to add to the pose after the player evaluates
 	TBaseBlendedCurve<FDefaultAllocator, UE::Anim::FCurveElement> OverrideCurve;
 
@@ -143,7 +145,7 @@ struct BLENDSTACK_API FAnimNode_BlendStack_Standalone : public FAnimNode_AssetPl
 		bool bMirrored = false, UMirrorDataTable* MirrorDataTable = nullptr, float BlendTime = 0.2f,
 		const UBlendProfile* BlendProfile = nullptr, EAlphaBlendOption BlendOption = EAlphaBlendOption::Linear, 
 		bool bUseInertialBlend = false, const FVector& BlendParameters = FVector::Zero(), float PlayRate = 1.f, float ActivationDelay = 0.f,
-		FName GroupName = NAME_None, EAnimGroupRole::Type GroupRole = EAnimGroupRole::CanBeLeader, EAnimSyncMethod Method = EAnimSyncMethod::DoNotSync);
+		FName GroupName = NAME_None, EAnimGroupRole::Type GroupRole = EAnimGroupRole::CanBeLeader, EAnimSyncMethod Method = EAnimSyncMethod::DoNotSync, bool bOverridePositionWhenJoiningSyncGroupAsLeader = false);
 	void UpdatePlayRate(float PlayRate);
 	void Reset();
 
@@ -163,7 +165,7 @@ protected:
 		bool bMirrored, UMirrorDataTable* MirrorDataTable, float BlendTime,
 		const UBlendProfile* BlendProfile, EAlphaBlendOption BlendOption, 
 		bool bUseInertialBlend, const FVector& BlendParameters, float PlayRate, float ActivationDelay,
-		FName GroupName, EAnimGroupRole::Type GroupRole, EAnimSyncMethod Method);
+		FName GroupName, EAnimGroupRole::Type GroupRole, EAnimSyncMethod Method, bool bOverridePositionWhenJoiningSyncGroupAsLeader);
 
 	static void BlendWithPose(FAnimationPoseData& InOutPoseData, const FAnimationPoseData& OtherPoseData, const float InOutPoseWeight);
 	static void BlendWithPosePerBone(FAnimationPoseData& InOutPoseData, const FAnimationPoseData& OtherPoseData, TConstArrayView<float> OtherPoseWeights);
diff --git a/Engine/Plugins/Animation/PoseSearch/Source/Editor/Private/AnimGraphNode_MotionMatching.cpp b/Engine/Plugins/Animation/PoseSearch/Source/Editor/Private/AnimGraphNode_MotionMatching.cpp
--- a/Engine/Plugins/Animation/PoseSearch/Source/Editor/Private/AnimGraphNode_MotionMatching.cpp
+++ b/Engine/Plugins/Animation/PoseSearch/Source/Editor/Private/AnimGraphNode_MotionMatching.cpp
@@ -1,8 +1,14 @@
 // Copyright Epic Games, Inc. All Rights Reserved.
 
 #include "AnimGraphNode_MotionMatching.h"
+
+#include "DetailCategoryBuilder.h"
+#include "DetailLayoutBuilder.h"
 #include "Animation/AnimRootMotionProvider.h"
 #include "FindInBlueprintManager.h"
+#include "Widgets/SBoxPanel.h"
+#include "Widgets/Text/STextBlock.h"
+#include "Widgets/Images/SImage.h"
 
 #define LOCTEXT_NAMESPACE "AnimGraphNode_MotionMatching"
 
@@ -19,7 +25,9 @@ FText UAnimGraphNode_MotionMatching::GetTooltipText() const
 
 FText UAnimGraphNode_MotionMatching::GetNodeTitle(ENodeTitleType::Type TitleType) const
 {
-	return LOCTEXT("NodeTitle", "Motion Matching");
+	FText Title = LOCTEXT("NodeTitle", "Motion Matching");
+	AddSyncGroupToNodeTitle(TitleType, Title);
+	return Title;
 }
 
 FText UAnimGraphNode_MotionMatching::GetMenuCategory() const
@@ -59,6 +67,40 @@ void UAnimGraphNode_MotionMatching::PostEditChangeProperty(FPropertyChangedEvent
 	Super::PostEditChangeProperty(PropertyChangedEvent);
 }
 
+void UAnimGraphNode_MotionMatching::CustomizeDetails(IDetailLayoutBuilder& DetailBuilder)
+{
+	Super::CustomizeDetails(DetailBuilder);
+	DetailBuilder.EditCategory("Sync").SetToolTip(LOCTEXT("CategoryToolTip", "Motion Matching is intended to only act as a leader in a sync group. Internally, only the most recent blend stack sample will have sync information for other asset players to follow and the rest of blend stack samples will tick without syncing."));
+	
+	auto HeaderContentWidget = SNew(SHorizontalBox)
+	+ SHorizontalBox::Slot()
+	.AutoWidth()
+	.HAlign(HAlign_Center)
+	.Padding(4.0f, 4.0f)
+	.VAlign(VAlign_Center)
+	[
+		SNew(STextBlock)
+		.Text(LOCTEXT("SyncTextLabel", "Sync"))
+		.Font(IDetailLayoutBuilder::GetDetailFontBold())
+		.TextStyle(FAppStyle::Get(), "DetailsView.CategoryTextStyle")
+	]
+	
+	+ SHorizontalBox::Slot()
+	.AutoWidth()
+	.HAlign(HAlign_Left)
+	.VAlign(VAlign_Center)
+	[
+		SNew(SImage)
+		.DesiredSizeOverride(FVector2D(13.f, 13.f))
+		.ColorAndOpacity(FSlateColor::UseSubduedForeground())
+		.Image(FAppStyle::Get().GetBrush("Icons.Info"))
+		.ToolTipText(LOCTEXT("SyncCategoryToolTip", "Motion Matching is intended to only act as a leader in a sync group. Internally, only the most recent blend stack sample will have sync information for other asset players to follow and the rest of blend stack samples will tick without syncing."))
+	];
+	
+	DetailBuilder.EditCategory("Sync").HeaderContent(HeaderContentWidget, true);
+	DetailBuilder.EditCategory("Sync").SetSortOrder(0);
+}
+
 void UAnimGraphNode_MotionMatching::AddSearchMetaDataInfo(TArray<struct FSearchTagDataPair>& OutTaggedMetaData) const
 {
 	Super::AddSearchMetaDataInfo(OutTaggedMetaData);
diff --git a/Engine/Plugins/Animation/PoseSearch/Source/Editor/Public/AnimGraphNode_MotionMatching.h b/Engine/Plugins/Animation/PoseSearch/Source/Editor/Public/AnimGraphNode_MotionMatching.h
--- a/Engine/Plugins/Animation/PoseSearch/Source/Editor/Public/AnimGraphNode_MotionMatching.h
+++ b/Engine/Plugins/Animation/PoseSearch/Source/Editor/Public/AnimGraphNode_MotionMatching.h
@@ -18,21 +18,25 @@ class UAnimGraphNode_MotionMatching : public UAnimGraphNode_BlendStack_Base
 	UPROPERTY(EditAnywhere, Category = Settings)
 	FAnimNode_MotionMatching Node;
 
+	// UAnimGraphNode_Base interface
 	virtual FLinearColor GetNodeTitleColor() const override;
 	virtual FText GetTooltipText() const override;
 	virtual FText GetNodeTitle(ENodeTitleType::Type TitleType) const override;
 	virtual FText GetMenuCategory() const override;
 	virtual UScriptStruct* GetTimePropertyStruct() const override;
-
-	virtual FAnimNode_BlendStack_Standalone* GetBlendStackNode() const override { return (FAnimNode_BlendStack_Standalone*)(&Node); }
-
 	virtual void OnProcessDuringCompilation(IAnimBlueprintCompilationContext& InCompilationContext, IAnimBlueprintGeneratedClassCompiledData& OutCompiledData) override;
 	virtual void ValidateAnimNodeDuringCompilation(USkeleton* ForSkeleton, FCompilerResultsLog& MessageLog) override;
 	virtual void AddSearchMetaDataInfo(TArray<struct FSearchTagDataPair>& OutTaggedMetaData) const override;
 	virtual void GetBoundFunctionsInfo(TArray<TPair<FName, FName>>& InOutBindingsInfo) override;
 	virtual bool ReferencesFunction(const FName& InFunctionName, const UStruct* InScope) const override;
 	virtual void PostEditChangeProperty(FPropertyChangedEvent& PropertyChangedEvent) override;
+	virtual void CustomizeDetails(IDetailLayoutBuilder& DetailBuilder) override;
+	// End of UEdGraphNode interface
 
+	// UAnimGraphNode_BlendStack_Base interface
+	virtual FAnimNode_BlendStack_Standalone* GetBlendStackNode() const override { return (FAnimNode_BlendStack_Standalone*)(&Node); }
+	// End of UAnimGraphNode_BlendStack_Base interface
+	
 	UPROPERTY(EditAnywhere, Category = "Functions|Motion Matching", meta=(FunctionReference, AllowFunctionLibraries, PrototypeFunction="/Script/AnimGraphRuntime.AnimExecutionContextLibrary.Prototype_ThreadSafeAnimUpdateCall", DefaultBindingName="OnStateExit"), DisplayName="On Motion Matching State Updated")
 	FMemberReference OnMotionMatchingStateUpdatedFunction;
 };
diff --git a/Engine/Plugins/Animation/PoseSearch/Source/Runtime/Private/AnimNode_MotionMatching.cpp b/Engine/Plugins/Animation/PoseSearch/Source/Runtime/Private/AnimNode_MotionMatching.cpp
--- a/Engine/Plugins/Animation/PoseSearch/Source/Runtime/Private/AnimNode_MotionMatching.cpp
+++ b/Engine/Plugins/Animation/PoseSearch/Source/Runtime/Private/AnimNode_MotionMatching.cpp
@@ -199,9 +199,21 @@ void FAnimNode_MotionMatching::UpdateAssetPlayer(const FAnimationUpdateContext&
 
 			if (UAnimationAsset* AnimationAsset = Cast<UAnimationAsset>(DatabaseAsset->GetAnimationAsset()))
 			{
+				// Clear up any sync group info before pushing new asset player (which will have sync info since its the highest weighted).
+				for (FBlendStackAnimPlayer& AnimPlayer : AnimPlayers)
+				{
+					if (FAnimNode_AssetPlayerBase* AssetPlayerNode = AnimPlayer.GetAssetPlayerNode())
+					{
+						AssetPlayerNode->SetGroupMethod(EAnimSyncMethod::DoNotSync);
+						AssetPlayerNode->SetGroupRole(EAnimGroupRole::CanBeLeader);
+						AssetPlayerNode->SetGroupName(NAME_None);
+					}
+				}
+				
 				FAnimNode_BlendStack_Standalone::BlendTo(Context, AnimationAsset, MotionMatchingState.CurrentSearchResult.AssetTime,
 					SearchIndexAsset->IsLooping(), SearchIndexAsset->IsMirrored(), CurrentResultDatabase->Schema->GetMirrorDataTable(DefaultRole), BlendTime,
-					BlendProfile, BlendOption, bUseInertialBlend, SearchIndexAsset->GetBlendParameters(), MotionMatchingState.WantedPlayRate * PlayRateMultiplier);
+					BlendProfile, BlendOption, bUseInertialBlend, SearchIndexAsset->GetBlendParameters(), MotionMatchingState.WantedPlayRate * PlayRateMultiplier, 0,
+					GetGroupName(), GetGroupRole(), GetGroupMethod(), GetOverridePositionWhenJoiningSyncGroupAsLeader());
 			}
 			else
 			{
@@ -278,4 +290,86 @@ bool FAnimNode_MotionMatching::SetIgnoreForRelevancyTest(bool bInIgnoreForReleva
 	return false;
 }
 
+FName FAnimNode_MotionMatching::GetGroupName() const
+{
+	return GET_ANIM_NODE_DATA(FName, GroupName);
+}
+
+EAnimGroupRole::Type FAnimNode_MotionMatching::GetGroupRole() const
+{
+	return GET_ANIM_NODE_DATA(TEnumAsByte<EAnimGroupRole::Type>, GroupRole);
+}
+
+EAnimSyncMethod FAnimNode_MotionMatching::GetGroupMethod() const
+{
+	return GET_ANIM_NODE_DATA(EAnimSyncMethod, Method);
+}
+
+bool FAnimNode_MotionMatching::GetOverridePositionWhenJoiningSyncGroupAsLeader() const
+{
+	return GET_ANIM_NODE_DATA(bool, bOverridePositionWhenJoiningSyncGroupAsLeader);
+}
+
+bool FAnimNode_MotionMatching::IsLooping() const
+{
+	if (!AnimPlayers.IsEmpty())
+	{
+		return AnimPlayers[0].IsLooping();
+	}
+	return false;
+}
+
+bool FAnimNode_MotionMatching::SetGroupName(FName InGroupName)
+{
+#if WITH_EDITORONLY_DATA
+	GroupName = InGroupName;
+#endif
+	if(FName* GroupNamePtr = GET_INSTANCE_ANIM_NODE_DATA_PTR(FName, GroupName))
+	{
+		*GroupNamePtr = InGroupName;
+		return true;
+	}
+	return false;
+}
+
+bool FAnimNode_MotionMatching::SetGroupRole(EAnimGroupRole::Type InRole)
+{
+#if WITH_EDITORONLY_DATA
+	GroupRole = InRole;
+#endif
+	
+	if(TEnumAsByte<EAnimGroupRole::Type>* GroupRolePtr = GET_INSTANCE_ANIM_NODE_DATA_PTR(TEnumAsByte<EAnimGroupRole::Type>, GroupRole))
+	{
+		*GroupRolePtr = InRole;
+		return true;
+	}
+	return false;
+}
+
+bool FAnimNode_MotionMatching::SetGroupMethod(EAnimSyncMethod InMethod)
+{
+#if WITH_EDITORONLY_DATA
+	Method = InMethod;
+#endif
+	if(EAnimSyncMethod* MethodPtr = GET_INSTANCE_ANIM_NODE_DATA_PTR(EAnimSyncMethod, Method))
+	{
+		*MethodPtr = InMethod;
+		return true;
+	}
+	return false;
+}
+
+bool FAnimNode_MotionMatching::SetOverridePositionWhenJoiningSyncGroupAsLeader(bool InOverridePositionWhenJoiningSyncGroupAsLeader)
+{
+#if WITH_EDITORONLY_DATA
+	bOverridePositionWhenJoiningSyncGroupAsLeader = InOverridePositionWhenJoiningSyncGroupAsLeader;
+#endif
+	if(bool* bOverridePositionWhenJoiningSyncGroupAsLeaderPtr = GET_INSTANCE_ANIM_NODE_DATA_PTR(bool, bOverridePositionWhenJoiningSyncGroupAsLeader))
+	{
+		*bOverridePositionWhenJoiningSyncGroupAsLeaderPtr = InOverridePositionWhenJoiningSyncGroupAsLeader;
+		return true;
+	}
+	return false;
+}
+
 #undef LOCTEXT_NAMESPACE
diff --git a/Engine/Plugins/Animation/PoseSearch/Source/Runtime/Public/PoseSearch/AnimNode_MotionMatching.h b/Engine/Plugins/Animation/PoseSearch/Source/Runtime/Public/PoseSearch/AnimNode_MotionMatching.h
--- a/Engine/Plugins/Animation/PoseSearch/Source/Runtime/Public/PoseSearch/AnimNode_MotionMatching.h
+++ b/Engine/Plugins/Animation/PoseSearch/Source/Runtime/Public/PoseSearch/AnimNode_MotionMatching.h
@@ -39,14 +39,51 @@ private:
 	// @todo: implement CacheBones_AnyThread to rebind the schema bones
 	virtual void Initialize_AnyThread(const FAnimationInitializeContext& Context) override;
 	virtual void Evaluate_AnyThread(FPoseContext& Output) override;
+	virtual bool GetIgnoreForRelevancyTest() const override;
+	virtual bool SetIgnoreForRelevancyTest(bool bInIgnoreForRelevancyTest) override;
+	virtual FName GetGroupName() const override;
+	virtual EAnimGroupRole::Type GetGroupRole() const override;
+	virtual EAnimSyncMethod GetGroupMethod() const override;
+	virtual bool GetOverridePositionWhenJoiningSyncGroupAsLeader() const override;
+	virtual bool IsLooping() const override;
+	virtual bool SetGroupName(FName InGroupName) override;
+	virtual bool SetGroupRole(EAnimGroupRole::Type InRole) override;
+	virtual bool SetGroupMethod(EAnimSyncMethod InMethod) override;
+	virtual bool SetOverridePositionWhenJoiningSyncGroupAsLeader(bool InOverridePositionWhenJoiningSyncGroupAsLeader) override;
 	// End of FAnimNode_Base interface
 
 	// FAnimNode_AssetPlayerBase interface
 	virtual void UpdateAssetPlayer(const FAnimationUpdateContext& Context) override;
-	virtual bool GetIgnoreForRelevancyTest() const override;
-	virtual bool SetIgnoreForRelevancyTest(bool bInIgnoreForRelevancyTest) override;
 	// End of FAnimNode_AssetPlayerBase interface
 
+#if WITH_EDITORONLY_DATA
+	
+	// The group name that we synchronize with (NAME_None if it is not part of any group). Note that
+	// this is the name of the group used to sync the output of this node - it will not force
+	// syncing of animations contained by it.
+	UPROPERTY(EditAnywhere, Category = Sync, meta = (FoldProperty))
+	FName GroupName = NAME_None;
+
+	// The role this node can assume within the group (ignored if GroupName is not set). Note
+	// that this is the role of the output of this node, not of animations contained by it.
+	UPROPERTY(VisibleAnywhere, Category = Sync, meta = (FoldProperty))
+	TEnumAsByte<EAnimGroupRole::Type> GroupRole = EAnimGroupRole::ExclusiveAlwaysLeader;
+	
+	// When enabled, acting as the leader, and using marker-based sync, this asset player will not sync to the previous leader's sync position when joining a sync group and before becoming the leader but instead force everyone else to match its position.
+	UPROPERTY(VisibleAnywhere, Category = Sync, meta = (FoldProperty, EditCondition = "GroupRole == EAnimGroupRole::AlwaysLeader || GroupRole == EAnimGroupRole::ExclusiveAlwaysLeader || GroupRole == EAnimGroupRole::TransitionLeader", EditConditionHides))
+	bool bOverridePositionWhenJoiningSyncGroupAsLeader = true;
+	
+	// How this node will synchronize with other animations. Note that this determines how the output
+	// of this node is used for synchronization, not of animations contained by it.
+	UPROPERTY(EditAnywhere, Category = Sync, meta = (FoldProperty))
+	EAnimSyncMethod Method = EAnimSyncMethod::DoNotSync;
+
+	// If true, "Relevant anim" nodes that look for the highest weighted animation in a state will ignore this node
+	UPROPERTY(EditAnywhere, Category = Relevancy, meta = (FoldProperty, PinHiddenByDefault))
+	bool bIgnoreForRelevancyTest = false;
+
+#endif
+	
 	// The database to search. This can be overridden by Anim Node Functions such as "On Become Relevant" and "On Update" via SetDatabaseToSearch/SetDatabasesToSearch.
 	UPROPERTY(EditAnywhere, Category = Settings, meta = (PinShownByDefault))
 	TObjectPtr<const UPoseSearchDatabase> Database = nullptr;
@@ -117,10 +154,7 @@ private:
 	bool bOverrideDatabaseInput = false;
 
 #if WITH_EDITORONLY_DATA
-	// If true, "Relevant anim" nodes that look for the highest weighted animation in a state will ignore this node
-	UPROPERTY(EditAnywhere, Category = Relevancy, meta = (FoldProperty, PinHiddenByDefault))
-	bool bIgnoreForRelevancyTest = false;
-
+	
 	UPROPERTY(meta=(FoldProperty))
 	FAnimNodeFunctionRef OnMotionMatchingStateUpdated;
 
diff --git a/Engine/Plugins/Experimental/AnimNextAnimGraph/Source/AnimNextAnimGraph/Private/Graph/SyncGroup_GraphInstanceComponent.cpp b/Engine/Plugins/Experimental/AnimNextAnimGraph/Source/AnimNextAnimGraph/Private/Graph/SyncGroup_GraphInstanceComponent.cpp
--- a/Engine/Plugins/Experimental/AnimNextAnimGraph/Source/AnimNextAnimGraph/Private/Graph/SyncGroup_GraphInstanceComponent.cpp
+++ b/Engine/Plugins/Experimental/AnimNextAnimGraph/Source/AnimNextAnimGraph/Private/Graph/SyncGroup_GraphInstanceComponent.cpp
@@ -55,6 +55,7 @@ namespace UE::AnimNext
 					}
 					break;
 				case EAnimGroupRole::AlwaysLeader:
+				case EAnimGroupRole::ExclusiveAlwaysLeader:
 					// Always set the leader index
 					LeaderIndex = MemberIndex;
 					LeaderTotalWeight = 2.0f;		// Some high value
diff --git a/Engine/Source/Editor/AnimGraph/Private/AnimGraphNode_AssetPlayerBase.cpp b/Engine/Source/Editor/AnimGraph/Private/AnimGraphNode_AssetPlayerBase.cpp
--- a/Engine/Source/Editor/AnimGraph/Private/AnimGraphNode_AssetPlayerBase.cpp
+++ b/Engine/Source/Editor/AnimGraph/Private/AnimGraphNode_AssetPlayerBase.cpp
@@ -447,47 +447,52 @@ FText UAnimGraphNode_AssetPlayerBase::GetNodeTitleForAsset(ENodeTitleType::Type
 		TitleArgs.Add(TEXT("AssetDesc"), InAssetDesc);
 		FText Title = FText::Format(LOCTEXT("AssetPlayerFullTitle", "{AssetName}\n{AssetDesc}"), TitleArgs);
 
-		if (InTitleType == ENodeTitleType::FullTitle)
+		AddSyncGroupToNodeTitle(InTitleType, Title);
+
+		return Title;
+	}
+}
+
+void UAnimGraphNode_AssetPlayerBase::AddSyncGroupToNodeTitle(ENodeTitleType::Type InTitleType, FText& InOutTitle) const
+{
+	if (InTitleType == ENodeTitleType::FullTitle)
+	{
+		FStructProperty* NodeProperty = GetFNodeProperty();
+		if(NodeProperty->Struct->IsChildOf(FAnimNode_AssetPlayerBase::StaticStruct()))
 		{
-			FStructProperty* NodeProperty = GetFNodeProperty();
-			if(NodeProperty->Struct->IsChildOf(FAnimNode_AssetPlayerBase::StaticStruct()))
-			{
-				const FAnimNode_AssetPlayerBase* Node = NodeProperty->ContainerPtrToValuePtr<FAnimNode_AssetPlayerBase>(this);
+			const FAnimNode_AssetPlayerBase* Node = NodeProperty->ContainerPtrToValuePtr<FAnimNode_AssetPlayerBase>(this);
 				
-				FFormatNamedArguments Args;
-				Args.Add(TEXT("Title"), Title);
+			FFormatNamedArguments Args;
+			Args.Add(TEXT("Title"), InOutTitle);
 
-				if(Node->GetGroupMethod() == EAnimSyncMethod::SyncGroup)
-				{
-					Args.Add(TEXT("SyncGroupName"), FText::FromName(Node->GetGroupName()));
-					static const FTextFormat FormatAssetPlayerNodeSyncGroupSubtitle(LOCTEXT("AssetPlayerNodeSyncGroupSubtitle", "{Title}\nSync group {SyncGroupName}"));
-					Title = FText::Format(FormatAssetPlayerNodeSyncGroupSubtitle, Args);
-				}
-				else if(Node->GetGroupMethod() == EAnimSyncMethod::Graph)
-				{
-					static const FTextFormat FormatAssetPlayerNodeGraphSyncGroupSubtitle(LOCTEXT("AssetPlayerNodeGraphSyncGroupSubtitle", "{Title}\nGraph sync group"));
-					Title = FText::Format(FormatAssetPlayerNodeGraphSyncGroupSubtitle, Args);
+			if(Node->GetGroupMethod() == EAnimSyncMethod::SyncGroup)
+			{
+				Args.Add(TEXT("SyncGroupName"), FText::FromName(Node->GetGroupName()));
+				static const FTextFormat FormatAssetPlayerNodeSyncGroupSubtitle(LOCTEXT("AssetPlayerNodeSyncGroupSubtitle", "{Title}\nSync group {SyncGroupName}"));
+				InOutTitle = FText::Format(FormatAssetPlayerNodeSyncGroupSubtitle, Args);
+			}
+			else if(Node->GetGroupMethod() == EAnimSyncMethod::Graph)
+			{
+				static const FTextFormat FormatAssetPlayerNodeGraphSyncGroupSubtitle(LOCTEXT("AssetPlayerNodeGraphSyncGroupSubtitle", "{Title}\nGraph sync group"));
+				InOutTitle = FText::Format(FormatAssetPlayerNodeGraphSyncGroupSubtitle, Args);
 
-					UObject* ObjectBeingDebugged = GetAnimBlueprint()->GetObjectBeingDebugged();
-					UAnimBlueprintGeneratedClass* GeneratedClass = GetAnimBlueprint()->GetAnimBlueprintGeneratedClass();
-					if (ObjectBeingDebugged && GeneratedClass)
+				UObject* ObjectBeingDebugged = GetAnimBlueprint()->GetObjectBeingDebugged();
+				UAnimBlueprintGeneratedClass* GeneratedClass = GetAnimBlueprint()->GetAnimBlueprintGeneratedClass();
+				if (ObjectBeingDebugged && GeneratedClass)
+				{
+					int32 NodeIndex = GeneratedClass->GetNodeIndexFromGuid(NodeGuid);
+					if(NodeIndex != INDEX_NONE)
 					{
-						int32 NodeIndex = GeneratedClass->GetNodeIndexFromGuid(NodeGuid);
-						if(NodeIndex != INDEX_NONE)
+						if(const FName* SyncGroupNamePtr = GeneratedClass->GetAnimBlueprintDebugData().NodeSyncsThisFrame.Find(NodeIndex))
 						{
-							if(const FName* SyncGroupNamePtr = GeneratedClass->GetAnimBlueprintDebugData().NodeSyncsThisFrame.Find(NodeIndex))
-							{
-								Args.Add(TEXT("SyncGroupName"), FText::FromName(*SyncGroupNamePtr));
-								static const FTextFormat FormatAssetPlayerNodeDynamicGraphSyncGroupSubtitle(LOCTEXT("AssetPlayerNodeDynamicGraphSyncGroupSubtitle", "{Title}\nGraph sync group {SyncGroupName}"));
-								Title = FText::Format(FormatAssetPlayerNodeDynamicGraphSyncGroupSubtitle, Args);
-							}
+							Args.Add(TEXT("SyncGroupName"), FText::FromName(*SyncGroupNamePtr));
+							static const FTextFormat FormatAssetPlayerNodeDynamicGraphSyncGroupSubtitle(LOCTEXT("AssetPlayerNodeDynamicGraphSyncGroupSubtitle", "{Title}\nGraph sync group {SyncGroupName}"));
+							InOutTitle = FText::Format(FormatAssetPlayerNodeDynamicGraphSyncGroupSubtitle, Args);
 						}
 					}
 				}
 			}
 		}
-
-		return Title;
 	}
 }
 
diff --git a/Engine/Source/Editor/AnimGraph/Public/AnimGraphNode_AssetPlayerBase.h b/Engine/Source/Editor/AnimGraph/Public/AnimGraphNode_AssetPlayerBase.h
--- a/Engine/Source/Editor/AnimGraph/Public/AnimGraphNode_AssetPlayerBase.h
+++ b/Engine/Source/Editor/AnimGraph/Public/AnimGraphNode_AssetPlayerBase.h
@@ -68,6 +68,7 @@ protected:
 	// Helper functions to build a title for an asset player node
 	FText GetNodeTitleHelper(ENodeTitleType::Type InTitleType, UEdGraphPin* InAssetPin, const FText& InAssetDesc, const TFunction<FText(UAnimationAsset*)> InPostFixFunctionRef = nullptr) const;
 	FText GetNodeTitleForAsset(ENodeTitleType::Type InTitleType, UAnimationAsset* InAsset, const FText& InAssetDesc, const TFunction<FText(UAnimationAsset*)> InPostFixFunctionRef = nullptr) const;
+	void AddSyncGroupToNodeTitle(ENodeTitleType::Type InTitleType, FText & InOutTitle) const;
 	
 	// Default setup function that can be used with GetMenuActionsHalper
 	static void SetupNewNode(UEdGraphNode* InNewNode, bool bInIsTemplateNode, const FAssetData InAssetData);
diff --git a/Engine/Source/Runtime/AnimGraphRuntime/Public/AnimNodes/AnimNode_BlendSpacePlayer.h b/Engine/Source/Runtime/AnimGraphRuntime/Public/AnimNodes/AnimNode_BlendSpacePlayer.h
--- a/Engine/Source/Runtime/AnimGraphRuntime/Public/AnimNodes/AnimNode_BlendSpacePlayer.h
+++ b/Engine/Source/Runtime/AnimGraphRuntime/Public/AnimNodes/AnimNode_BlendSpacePlayer.h
@@ -228,6 +228,10 @@ private:
 	UPROPERTY(EditAnywhere, Category = Sync)
 	TEnumAsByte<EAnimGroupRole::Type> GroupRole = EAnimGroupRole::CanBeLeader;
 
+	// When enabled, acting as the leader, and using marker-based sync, this asset player will not sync to the previous leader's sync position when joining a sync group and before becoming the leader but instead force everyone else to match its position.
+	UPROPERTY(VisibleAnywhere, Category = Sync, meta = (EditCondition = "GroupRole == EAnimGroupRole::AlwaysLeader || GroupRole == EAnimGroupRole::ExclusiveAlwaysLeader || GroupRole == EAnimGroupRole::TransitionLeader", EditConditionHides))
+	bool bOverridePositionWhenJoiningSyncGroupAsLeader = true;
+	
 	// How this asset will synchronize with other assets. Note that this determines how the output
 	// of this node is used for synchronization, not of animations contained by it.
 	UPROPERTY(EditAnywhere, Category = Sync)
@@ -271,11 +275,13 @@ public:
 	virtual FName GetGroupName() const override { return GroupName; }
 	virtual EAnimGroupRole::Type GetGroupRole() const override { return GroupRole; }
 	virtual EAnimSyncMethod GetGroupMethod() const override { return Method; }
+	virtual bool GetOverridePositionWhenJoiningSyncGroupAsLeader() const override { return bOverridePositionWhenJoiningSyncGroupAsLeader; };
 	virtual bool GetIgnoreForRelevancyTest() const override { return bIgnoreForRelevancyTest; }
 	virtual bool IsLooping() const override { return bLoop; }
 	virtual bool SetGroupName(FName InGroupName) override { GroupName = InGroupName; return true; }
 	virtual bool SetGroupRole(EAnimGroupRole::Type InRole) override { GroupRole = InRole; return true; }
 	virtual bool SetGroupMethod(EAnimSyncMethod InMethod) override { Method = InMethod; return true; }
+	virtual bool SetOverridePositionWhenJoiningSyncGroupAsLeader(bool InOverridePositionWhenJoiningSyncGroupAsLeader) override { bOverridePositionWhenJoiningSyncGroupAsLeader = InOverridePositionWhenJoiningSyncGroupAsLeader; return true; };
 	virtual bool SetIgnoreForRelevancyTest(bool bInIgnoreForRelevancyTest) override { bIgnoreForRelevancyTest = bInIgnoreForRelevancyTest; return true; }
 	// End of FAnimNode_AssetPlayerBase interface
 
diff --git a/Engine/Source/Runtime/Engine/Classes/Animation/AnimNode_AssetPlayerBase.h b/Engine/Source/Runtime/Engine/Classes/Animation/AnimNode_AssetPlayerBase.h
--- a/Engine/Source/Runtime/Engine/Classes/Animation/AnimNode_AssetPlayerBase.h
+++ b/Engine/Source/Runtime/Engine/Classes/Animation/AnimNode_AssetPlayerBase.h
@@ -42,6 +42,9 @@ struct FAnimNode_AssetPlayerBase : public FAnimNode_AssetPlayerRelevancyBase
 	// Get the sync group method we are using
 	virtual EAnimSyncMethod GetGroupMethod() const { return EAnimSyncMethod::DoNotSync; }
 
+	// Get the flag that determines if this asset player will not sync to the previous leader's sync position when joining a sync group and before becoming the leader but instead force everyone else to match its position.
+	virtual bool GetOverridePositionWhenJoiningSyncGroupAsLeader() const { return false; }
+
 	// Set the sync group name we are using
 	virtual bool SetGroupName(FName InGroupName) { return false; }
 
@@ -51,6 +54,9 @@ struct FAnimNode_AssetPlayerBase : public FAnimNode_AssetPlayerRelevancyBase
 	// Set the sync group method we are using
 	virtual bool SetGroupMethod(EAnimSyncMethod InMethod) { return false; }
 
+	// Set the flag that determines if this asset player will not sync to the previous leader's sync position when joining a sync group and before becoming the leader but instead force everyone else to match its position. 
+	virtual bool SetOverridePositionWhenJoiningSyncGroupAsLeader(bool InOverridePositionWhenJoiningSyncGroupAsLeader) { return false; }
+	
 	// --- FAnimNode_RelevantAssetPlayerBase ---
 	ENGINE_API virtual float GetAccumulatedTime() const override;
 	ENGINE_API virtual void SetAccumulatedTime(float NewTime) override;
diff --git a/Engine/Source/Runtime/Engine/Classes/Animation/AnimNode_SequencePlayer.h b/Engine/Source/Runtime/Engine/Classes/Animation/AnimNode_SequencePlayer.h
--- a/Engine/Source/Runtime/Engine/Classes/Animation/AnimNode_SequencePlayer.h
+++ b/Engine/Source/Runtime/Engine/Classes/Animation/AnimNode_SequencePlayer.h
@@ -182,6 +182,10 @@ protected:
 	UPROPERTY(EditAnywhere, Category=Sync)
 	TEnumAsByte<EAnimGroupRole::Type> GroupRole = EAnimGroupRole::CanBeLeader;
 
+	// When enabled, acting as the leader, and using marker-based sync, this asset player will not sync to the previous leader's sync position when joining a sync group and before becoming the leader but instead force everyone else to match its position.
+	UPROPERTY(VisibleAnywhere, Category = Sync, meta = (EditCondition = "GroupRole == EAnimGroupRole::AlwaysLeader || GroupRole == EAnimGroupRole::ExclusiveAlwaysLeader || GroupRole == EAnimGroupRole::TransitionLeader", EditConditionHides))
+	bool bOverridePositionWhenJoiningSyncGroupAsLeader = true;
+	
 	// How this node will synchronize with other animations.
 	UPROPERTY(EditAnywhere, Category=Sync)
 	EAnimSyncMethod Method = EAnimSyncMethod::DoNotSync;
@@ -235,11 +239,13 @@ public:
 	virtual FName GetGroupName() const override { return GroupName; }
 	virtual EAnimGroupRole::Type GetGroupRole() const override { return GroupRole; }
 	virtual EAnimSyncMethod GetGroupMethod() const override { return Method; }
+	virtual bool GetOverridePositionWhenJoiningSyncGroupAsLeader() const override { return bOverridePositionWhenJoiningSyncGroupAsLeader; }
 	virtual bool IsLooping() const override { return bLoopAnimation; }
 	virtual bool GetIgnoreForRelevancyTest() const override { return bIgnoreForRelevancyTest; }
 	virtual bool SetGroupName(FName InGroupName) override { GroupName = InGroupName; return true; }
 	virtual bool SetGroupRole(EAnimGroupRole::Type InRole) override { GroupRole = InRole; return true; }
 	virtual bool SetGroupMethod(EAnimSyncMethod InMethod) override { Method = InMethod; return true; }
+	virtual bool SetOverridePositionWhenJoiningSyncGroupAsLeader(bool InOverridePositionWhenJoiningSyncGroupAsLeader) override { bOverridePositionWhenJoiningSyncGroupAsLeader = InOverridePositionWhenJoiningSyncGroupAsLeader; return true; }	
 	virtual bool SetIgnoreForRelevancyTest(bool bInIgnoreForRelevancyTest) override { bIgnoreForRelevancyTest = bInIgnoreForRelevancyTest; return true; }
 	virtual bool SetStartPosition(float InStartPosition) override { StartPosition = InStartPosition; return  true; }
 	virtual bool SetPlayRate(float InPlayRate) override { PlayRate = InPlayRate; return true; }
diff --git a/Engine/Source/Runtime/Engine/Classes/Animation/AnimationAsset.h b/Engine/Source/Runtime/Engine/Classes/Animation/AnimationAsset.h
--- a/Engine/Source/Runtime/Engine/Classes/Animation/AnimationAsset.h
+++ b/Engine/Source/Runtime/Engine/Classes/Animation/AnimationAsset.h
@@ -402,7 +402,9 @@ struct FAnimTickRecord
 	bool bLooping = false;
 	bool bIsEvaluator = false;
 	bool bRequestedInertialization = false;
-
+	bool bOverridePositionWhenJoiningSyncGroupAsLeader = false;
+	bool bIsExclusiveLeader = false;
+	
 	const UMirrorDataTable* MirrorDataTable = nullptr;
 
 	TSharedPtr<TArray<TUniquePtr<const UE::Anim::IAnimNotifyEventContextDataInterface>>> ContextData;
@@ -583,6 +585,9 @@ namespace EAnimGroupRole
 
 		/** This node will be excluded from the sync group while blending in. Once blended in it will be a follower until blended out*/
 		TransitionFollower,
+
+		/** This node will always be a leader. If it fails to be ticked as a leader it will be run as ungrouped asset player (EAnimSyncMethod::DoNotSync) .*/
+		ExclusiveAlwaysLeader,
 	};
 }
 
diff --git a/Engine/Source/Runtime/Engine/Private/Animation/AnimNode_AssetPlayerBase.cpp b/Engine/Source/Runtime/Engine/Private/Animation/AnimNode_AssetPlayerBase.cpp
--- a/Engine/Source/Runtime/Engine/Private/Animation/AnimNode_AssetPlayerBase.cpp
+++ b/Engine/Source/Runtime/Engine/Private/Animation/AnimNode_AssetPlayerBase.cpp
@@ -81,6 +81,7 @@ UE::Anim::FAnimSyncParams FAnimNode_AssetPlayerBase::GetSyncParams(bool bRequest
 	const FName SyncGroupName = GetGroupName();
 	FName GroupNameToUse = SyncGroupName;
 	EAnimSyncMethod MethodToUse = GetGroupMethod();
+	bool bOverridePositionWhenJoiningSyncGroupAsLeader = GetOverridePositionWhenJoiningSyncGroupAsLeader();
 	
 	// Skip sync based on roles.
 	{
@@ -105,5 +106,5 @@ UE::Anim::FAnimSyncParams FAnimNode_AssetPlayerBase::GetSyncParams(bool bRequest
 		}
 	}
 
-	return UE::Anim::FAnimSyncParams(GroupNameToUse, SyncGroupRole, MethodToUse);
+	return UE::Anim::FAnimSyncParams(GroupNameToUse, SyncGroupRole, MethodToUse, bOverridePositionWhenJoiningSyncGroupAsLeader);
 }
diff --git a/Engine/Source/Runtime/Engine/Private/Animation/AnimSync.cpp b/Engine/Source/Runtime/Engine/Private/Animation/AnimSync.cpp
--- a/Engine/Source/Runtime/Engine/Private/Animation/AnimSync.cpp
+++ b/Engine/Source/Runtime/Engine/Private/Animation/AnimSync.cpp
@@ -55,6 +55,8 @@ void FAnimSync::AddTickRecord(const FAnimTickRecord& InTickRecord, const FAnimSy
 		// Add animation instance's tick record to the sync group. 
 		SyncGroupInstance.ActivePlayers.Add(InTickRecord);
 		SyncGroupInstance.ActivePlayers.Top().MirrorDataTable = MirrorDataTable;
+		SyncGroupInstance.ActivePlayers.Top().bOverridePositionWhenJoiningSyncGroupAsLeader = InSyncParams.bOverridePositionWhenJoiningSyncGroupAsLeader;
+		SyncGroupInstance.ActivePlayers.Top().bIsExclusiveLeader = InSyncParams.Role == EAnimGroupRole::ExclusiveAlwaysLeader;
 
 		// Set leader score for the tick record we just added, and ensure there is only one montage per group.
 		// CanBeLeader or TransitionLeader Score (BlendWeight) < Always Leader Score (2.0) < Montage Leader Score (3.0).
@@ -165,25 +167,52 @@ void FAnimSync::TickAssetPlayerInstances(FAnimInstanceProxy& InProxy, float InDe
 					SCOPE_CYCLE_COUNTER(STAT_TickAssetPlayerInstance);
 					FScopeCycleCounterUObject Scope(GroupLeader.SourceAsset);
 
-					// Inertialization was requested therefore we resync to previous leader, if needed.
-                    if (GroupLeader.bRequestedInertialization)
-                    {
-                    	// Ensure we have a previous group leader to sync to otherwise we play normally. 
-                    	if (PreviousGroup && !PreviousGroup->ActivePlayers.IsEmpty())
-                    	{
-                    		// Only need to resync when using length based syncing since we initialized the tick context with the previous group's sync end position
-                    		// and that will take care of marker based syncing.
-                    		if (!TickContext.CanUseMarkerPosition() || !(PreviousGroup->ValidMarkers.Num() > 0))
-                    		{
-                    			const FAnimTickRecord & PreviousGroupLeader = PreviousGroup->ActivePlayers[PreviousGroup->GroupLeaderIndex];
-                    			const bool bShouldResyncToSyncGroup = PreviousGroupLeader.LeaderScore >= GroupLeader.LeaderScore;
+					// Enforce leader position and do not use previous group sync information as a starting point.
+					if (GroupLeader.bOverridePositionWhenJoiningSyncGroupAsLeader)
+					{
+						if (PreviousGroup && PreviousGroup->ActivePlayers.IsValidIndex(PreviousGroup->GroupLeaderIndex))
+						{
+							const FAnimTickRecord & PrevLeader = PreviousGroup->ActivePlayers[PreviousGroup->GroupLeaderIndex];
+
+							if (TickContext.CanUseMarkerPosition() && !PreviousGroup->ValidMarkers.IsEmpty())
+							{
+								const bool bIsAssetPlayerJustJoiningGroup = PrevLeader.MarkerTickRecord != GroupLeader.MarkerTickRecord || PrevLeader.DeltaTimeRecord != GroupLeader.DeltaTimeRecord;
+								const bool bDidLeaderAssetChanged = PrevLeader.SourceAsset != GroupLeader.SourceAsset;
+								
+								if (bIsAssetPlayerJustJoiningGroup || bDidLeaderAssetChanged)
+								{
+									TickContext.MarkerTickContext.SetMarkerSyncStartPosition(FMarkerSyncAnimPosition(NAME_None, NAME_None, 0));	
+								}
+							}
+							else
+							{
+								// Ensure no re-syncing occurs in order to overwrite previous leader position.
+								TickContext.SetResyncToSyncGroup(false);
+							}
+						}
+					}
+					else
+					{
+						// Inertialization was requested therefore we resync to previous leader, if needed.
+						if (GroupLeader.bRequestedInertialization)
+						{
+							// Ensure we have a previous group leader to sync to otherwise we play normally. 
+							if (PreviousGroup && !PreviousGroup->ActivePlayers.IsEmpty())
+							{
+								// Only need to resync when using length based syncing since we initialized the tick context with the previous group's sync end position
+								// and that will take care of marker based syncing.
+								if (!TickContext.CanUseMarkerPosition() || !(PreviousGroup->ValidMarkers.Num() > 0))
+								{
+									const FAnimTickRecord & PreviousGroupLeader = PreviousGroup->ActivePlayers[PreviousGroup->GroupLeaderIndex];
+									const bool bShouldResyncToSyncGroup = PreviousGroupLeader.LeaderScore >= GroupLeader.LeaderScore;
                 
-                    			// Sync to previous group leader if we have a lower score than them.
-                    			TickContext.SetResyncToSyncGroup(bShouldResyncToSyncGroup);
-                    		}
-                    	}
-                    }
-
+									// Sync to previous group leader if we have a lower score than them.
+									TickContext.SetResyncToSyncGroup(bShouldResyncToSyncGroup);
+								}
+							}
+						}
+					}
+					
 					// Tick group leader's asset.
 					TickContext.MarkerTickContext.MarkersPassedThisTick.Reset();
 					TickContext.RootMotionMovementParams.Clear();
@@ -232,7 +261,17 @@ void FAnimSync::TickAssetPlayerInstances(FAnimInstanceProxy& InProxy, float InDe
 				// Finalize sync group.
 				// This is where the followers tick records are reset if the needed. 
 				SyncGroup.Finalize(PreviousGroup);
-
+				
+				// Kick any records that requested being exclusive leaders but failed.
+				for (int32 RecordIndex = GroupLeaderIndex + 1; RecordIndex < SyncGroup.ActivePlayers.Num(); ++RecordIndex)
+				{
+					// We do not modify the array since we can just skip ticking their asset player and just push that to run as ungrouped.
+					if (SyncGroup.ActivePlayers[RecordIndex].bIsExclusiveLeader)
+					{
+						UngroupedActivePlayerArrays[GetSyncGroupWriteIndex()].Push(SyncGroup.ActivePlayers[RecordIndex]);
+					}
+				}
+				
 				// Ensure group leader's markers are valid, otherwise invalidate marker-based syncing for followers.
 				if (TickContext.CanUseMarkerPosition())
 				{
@@ -288,8 +327,15 @@ void FAnimSync::TickAssetPlayerInstances(FAnimInstanceProxy& InProxy, float InDe
 
 					for (int32 TickIndex = GroupLeaderIndex + 1; TickIndex < SyncGroup.ActivePlayers.Num(); ++TickIndex)
 					{
-						// Tick follower's asset player.
 						FAnimTickRecord& AssetPlayer = SyncGroup.ActivePlayers[TickIndex];
+
+						// Skip since we only wanted to be leaders. (Asset player will be ticked as ungrouped)
+						if (AssetPlayer.bIsExclusiveLeader)
+						{
+							continue;
+						}
+						
+						// Tick follower's asset player.
 						{
 							SCOPE_CYCLE_COUNTER(STAT_TickAssetPlayerInstance);
 							FScopeCycleCounterUObject Scope(AssetPlayer.SourceAsset);
diff --git a/Engine/Source/Runtime/Engine/Private/AnimationAsset.cpp b/Engine/Source/Runtime/Engine/Private/AnimationAsset.cpp
--- a/Engine/Source/Runtime/Engine/Private/AnimationAsset.cpp
+++ b/Engine/Source/Runtime/Engine/Private/AnimationAsset.cpp
@@ -87,6 +87,7 @@ void FAnimGroupInstance::TestTickRecordForLeadership(EAnimGroupRole::Type Member
 			Candidate.LeaderScore = Candidate.EffectiveBlendWeight;
 			break;
 		case EAnimGroupRole::AlwaysLeader:
+		case EAnimGroupRole::ExclusiveAlwaysLeader:
 			// Always set the leader index
 			Candidate.LeaderScore = LEADERSCORE_ALWAYSLEADER;
 			break;
diff --git a/Engine/Source/Runtime/Engine/Public/Animation/AnimSync.h b/Engine/Source/Runtime/Engine/Public/Animation/AnimSync.h
--- a/Engine/Source/Runtime/Engine/Public/Animation/AnimSync.h
+++ b/Engine/Source/Runtime/Engine/Public/Animation/AnimSync.h
@@ -16,15 +16,17 @@ namespace UE { namespace Anim {
 // Parameters for group/marker-based sync 
 struct FAnimSyncParams
 {
-	FAnimSyncParams(FName InGroupName = NAME_None, EAnimGroupRole::Type InRole = EAnimGroupRole::CanBeLeader, EAnimSyncMethod InMethod = EAnimSyncMethod::DoNotSync)
+	FAnimSyncParams(FName InGroupName = NAME_None, EAnimGroupRole::Type InRole = EAnimGroupRole::CanBeLeader, EAnimSyncMethod InMethod = EAnimSyncMethod::DoNotSync, bool InOverridePositionWhenJoiningSyncGroupAsLeader = false)
 		: GroupName(InGroupName)
 		, Role(InRole)
 		, Method(InMethod)
+		, bOverridePositionWhenJoiningSyncGroupAsLeader(InOverridePositionWhenJoiningSyncGroupAsLeader)
 	{}
-
+	
 	FName GroupName = NAME_None;
 	EAnimGroupRole::Type Role = EAnimGroupRole::CanBeLeader;
 	EAnimSyncMethod Method = EAnimSyncMethod::DoNotSync;
+	bool bOverridePositionWhenJoiningSyncGroupAsLeader = false;
 };
 
 // Wraps up functionality for ticking and syncing animations according to group (via normalized time) or marker
