blend stack - fix for blend profiles

diff --git a/Engine/Plugins/Animation/BlendStack/Source/Runtime/Private/AnimNode_BlendStack.cpp b/Engine/Plugins/Animation/BlendStack/Source/Runtime/Private/AnimNode_BlendStack.cpp
--- a/Engine/Plugins/Animation/BlendStack/Source/Runtime/Private/AnimNode_BlendStack.cpp
+++ b/Engine/Plugins/Animation/BlendStack/Source/Runtime/Private/AnimNode_BlendStack.cpp
@@ -50,9 +50,8 @@ void FBlendStackAnimPlayer::Initialize(const FAnimationInitializeContext& Contex
 		// handling BlendTime > 0 and RootBoneBlendTime >= 0
 		if (BlendProfile != nullptr)
 		{
-			TotalBlendInTimePerBone.Init(BlendTime, NumSkeletonBones);
+			TotalBlendInTimePerBone.SetNumUninitialized(NumSkeletonBones);
 			BlendProfile->FillSkeletonBoneDurationsArray(TotalBlendInTimePerBone, BlendTime, Skeleton);
-			BlendTime = *Algo::MaxElement(TotalBlendInTimePerBone);
 		}
 	}
 
@@ -399,7 +398,7 @@ UAnimationAsset* FBlendStackAnimPlayer::GetAnimationAsset() const
 
 float FBlendStackAnimPlayer::GetBlendInPercentage() const
 {
-	if (FMath::IsNearlyZero(TotalBlendInTime))
+	if (TotalBlendInTime < UE_SMALL_NUMBER)
 	{
 		if (TimeToActivation > 0.f)
 		{
@@ -408,7 +407,8 @@ float FBlendStackAnimPlayer::GetBlendInPercentage() const
 		return 1.f;
 	}
 
-	return FMath::Clamp(GetCurrentBlendInTime() / TotalBlendInTime, 0.f, 1.f);
+	check(CurrentBlendInTime >= 0.f);
+	return FMath::Min(CurrentBlendInTime / TotalBlendInTime, 1.f);
 }
 
 int32 FBlendStackAnimPlayer::GetBlendInWeightsNum() const
@@ -416,21 +416,30 @@ int32 FBlendStackAnimPlayer::GetBlendInWeightsNum() const
 	return TotalBlendInTimePerBone.Num();
 }
 
+float FBlendStackAnimPlayer::GetBlendInWeight() const
+{
+	const float BlendInPercentage = GetBlendInPercentage();
+	const float BlendInWeight = FAlphaBlend::AlphaToBlendOption(BlendInPercentage, GetBlendOption());
+	return BlendInWeight;
+}
+
 void FBlendStackAnimPlayer::GetBlendInWeights(TArrayView<float> Weights) const
 {
 	check(Weights.Num() == GetBlendInWeightsNum());
 	
+	const float WeightForZeroBlendInTime = TimeToActivation > 0.f ? 0.f : 1.f;
 	for (int32 BoneIdx = 0; BoneIdx < Weights.Num(); ++BoneIdx)
 	{
 		const float TotalBlendInTimeBoneIdx = TotalBlendInTimePerBone[BoneIdx];
-		if (FMath::IsNearlyZero(TotalBlendInTimeBoneIdx))
+		if (TotalBlendInTimeBoneIdx < UE_SMALL_NUMBER)
 		{
-			Weights[BoneIdx] = 1.f;
+			Weights[BoneIdx] = WeightForZeroBlendInTime;
 		}
 		else
 		{
-			const float UnclampedLinearWeight = GetCurrentBlendInTime() / TotalBlendInTimeBoneIdx;
-			Weights[BoneIdx] = FAlphaBlend::AlphaToBlendOption(UnclampedLinearWeight, BlendOption);
+			check(CurrentBlendInTime >= 0.f);
+			const float LinearWeight = FMath::Min(CurrentBlendInTime / TotalBlendInTimeBoneIdx, 1.f);
+			Weights[BoneIdx] = FAlphaBlend::AlphaToBlendOption(LinearWeight, BlendOption);
 		}
 	}
 }
@@ -538,14 +547,14 @@ void FAnimNode_BlendStack_Standalone::Evaluate_AnyThread(FPoseContext& Output)
 				const int32 BlendInWeightsNum = AnimPlayers[PlayerIndex].GetBlendInWeightsNum();
 				if (BlendInWeightsNum > 0)
 				{
-					TArrayView<float> Weights((float*)FMemory_Alloca(BlendInWeightsNum * sizeof(float)), BlendInWeightsNum);
-					AnimPlayers[PlayerIndex].GetBlendInWeights(Weights);
-					BlendWithPosePerBone(OutputAnimationPoseData, EvaluationAnimationPoseData, Weights);
+					TArrayView<float> EvaluationAnimationPoseDataWeights((float*)FMemory_Alloca(BlendInWeightsNum * sizeof(float)), BlendInWeightsNum);
+					AnimPlayers[PlayerIndex].GetBlendInWeights(EvaluationAnimationPoseDataWeights);
+					BlendWithPosePerBone(OutputAnimationPoseData, EvaluationAnimationPoseData, EvaluationAnimationPoseDataWeights);
 				}
 				else
 				{
-					const float Weight = 1.f - FAlphaBlend::AlphaToBlendOption(AnimPlayers[PlayerIndex].GetBlendInPercentage(), AnimPlayers[PlayerIndex].GetBlendOption());
-					BlendWithPose(OutputAnimationPoseData, EvaluationAnimationPoseData, Weight);
+					const float OutputAnimationPoseDataWeight = 1.f - AnimPlayers[PlayerIndex].GetBlendInWeight();
+					BlendWithPose(OutputAnimationPoseData, EvaluationAnimationPoseData, OutputAnimationPoseDataWeight);
 				}
 			}
 		};
@@ -665,7 +674,7 @@ void FAnimNode_BlendStack_Standalone::UpdateAssetPlayer(const FAnimationUpdateCo
 	{
 		FBlendStackAnimPlayer& AnimPlayer = AnimPlayers[AnimPlayerIndex];
 		const bool bIsLastAnimPlayers = AnimPlayerIndex == BlendStackSize - 1;
-		const float BlendInPercentage = bIsLastAnimPlayers ? 1.f : AnimPlayer.GetBlendInPercentage();
+		const float BlendInPercentage = bIsLastAnimPlayers ? 1.f : AnimPlayer.GetBlendInWeight();
 		const float AnimPlayerBlendWeight = CurrentWeightMultiplier * BlendInPercentage;
 
 		FAnimationUpdateContext AnimPlayerContext = Context.FractionalWeightAndRootMotion(AnimPlayerBlendWeight, AnimPlayerBlendWeight);
diff --git a/Engine/Plugins/Animation/BlendStack/Source/Runtime/Public/BlendStack/AnimNode_BlendStack.h b/Engine/Plugins/Animation/BlendStack/Source/Runtime/Public/BlendStack/AnimNode_BlendStack.h
--- a/Engine/Plugins/Animation/BlendStack/Source/Runtime/Public/BlendStack/AnimNode_BlendStack.h
+++ b/Engine/Plugins/Animation/BlendStack/Source/Runtime/Public/BlendStack/AnimNode_BlendStack.h
@@ -27,6 +27,7 @@ struct BLENDSTACK_API FBlendStackAnimPlayer
 	
 	int32 GetBlendInWeightsNum() const;
 	void GetBlendInWeights(TArrayView<float> Weights) const;
+	float GetBlendInWeight() const;
 
 	EAlphaBlendOption GetBlendOption() const { return BlendOption; }
 	void StorePoseContext(const FPoseContext& PoseContext);
