BlendStack: Fix for additive ref poses exploding

diff --git a/Engine/Plugins/Animation/BlendStack/Source/Runtime/Private/AnimNode_BlendStack.cpp b/Engine/Plugins/Animation/BlendStack/Source/Runtime/Private/AnimNode_BlendStack.cpp
--- a/Engine/Plugins/Animation/BlendStack/Source/Runtime/Private/AnimNode_BlendStack.cpp
+++ b/Engine/Plugins/Animation/BlendStack/Source/Runtime/Private/AnimNode_BlendStack.cpp
@@ -169,7 +169,7 @@ void FBlendStackAnimPlayer::RestorePoseContext(FPoseContext& PoseContext) const
 		{
 			if (StoredBones.IsEmpty())
 			{
-				PoseContext.Pose.ResetToRefPose();
+				PoseContext.ResetToRefPose();
 			}
 			else
 			{
@@ -199,7 +199,7 @@ void FBlendStackAnimPlayer::RestorePoseContext(FPoseContext& PoseContext) const
 	}
 	else
 	{
-		PoseContext.Pose.ResetToRefPose();
+		PoseContext.ResetToRefPose();
 	}
 	
 	PoseContext.Curve.CopyFrom(StoredCurve);
