BlendStack: Added option to update blend space input parameters.
Added delta threshold for re-blend when blend space input parameters change.

diff --git a/Engine/Plugins/Animation/BlendStack/Source/Runtime/Private/AnimNode_BlendStack.cpp b/Engine/Plugins/Animation/BlendStack/Source/Runtime/Private/AnimNode_BlendStack.cpp
--- a/Engine/Plugins/Animation/BlendStack/Source/Runtime/Private/AnimNode_BlendStack.cpp
+++ b/Engine/Plugins/Animation/BlendStack/Source/Runtime/Private/AnimNode_BlendStack.cpp
@@ -83,7 +83,7 @@ void FBlendStackAnimPlayer::Initialize(const FAnimationInitializeContext& Contex
 	else if (UBlendSpace* BlendSpace = Cast<UBlendSpace>(AnimationAsset))
 	{
 		// making sure AccumulatedTime is in normalized space
-		check(AccumulatedTime >= 0.f && AccumulatedTime <= 1.f);
+		AccumulatedTime = FMath::Clamp(AccumulatedTime, 0.f, 1.f);
 
 		BlendSpacePlayerNode.SetResetPlayTimeWhenBlendSpaceChanges(false /*!bReset*/);
 		BlendSpacePlayerNode.SetAccumulatedTime(AccumulatedTime);
@@ -446,6 +446,29 @@ void FBlendStackAnimPlayer::GetBlendInWeights(TArrayView<float> Weights) const
 
 /////////////////////////////////////////////////////
 // FAnimNode_BlendStack_Standalone
+
+
+void FAnimNode_BlendStack_Standalone::UpdateBlendspaceParameters(const EBlendStack_BlendspaceUpdateMode UpdateMode, const FVector& BlendParameters)
+{
+	// Update blend space parameters
+	if (UpdateMode == EBlendStack_BlendspaceUpdateMode::UpdateAll)
+	{
+		// apply blend space parameters to all blendspaces that are playing, including ones that are blending out
+		for (FBlendStackAnimPlayer& Player : AnimPlayers)
+		{
+			Player.SetBlendParameters(BlendParameters);
+		}
+	}
+	else if (UpdateMode == EBlendStack_BlendspaceUpdateMode::UpdateActiveOnly)
+	{
+		// apply blend space parameters only to the blendspace that is playing/blending in
+		if (!AnimPlayers.IsEmpty())
+		{
+			AnimPlayers[0].SetBlendParameters(BlendParameters);
+		}
+	}
+}
+
 void FAnimNode_BlendStack_Standalone::PopLastAnimPlayer()
 {
 	const int32 LastAnimPlayerIndex = AnimPlayers.Num() - 1;
@@ -1152,7 +1175,7 @@ void FAnimNode_BlendStack::UpdateAssetPlayer(const FAnimationUpdateContext& Cont
 		{
 			bExecuteBlendTo = true;
 		}
-		else if (BlendParameters != MainAnimPlayer.GetBlendParameters())
+		else if ((BlendParameters - MainAnimPlayer.GetBlendParameters()).SizeSquared() > FMath::Square(BlendParametersDeltaThreshold))
 		{
 			bExecuteBlendTo = true;
 		}
@@ -1173,6 +1196,8 @@ void FAnimNode_BlendStack::UpdateAssetPlayer(const FAnimationUpdateContext& Cont
 	UE::Anim::TOptionalScopedGraphMessage<UE::Anim::FAnimInertializationSyncScope> InertializationSync(bDidBlendToRequestAnInertialBlend, Context);
 	
 	UpdatePlayRate(WantedPlayRate);
+	UpdateBlendspaceParameters(BlendspaceUpdateMode, BlendParameters);
+
 	Super::UpdateAssetPlayer(Context);
 }
 
diff --git a/Engine/Plugins/Animation/BlendStack/Source/Runtime/Public/BlendStack/AnimNode_BlendStack.h b/Engine/Plugins/Animation/BlendStack/Source/Runtime/Public/BlendStack/AnimNode_BlendStack.h
--- a/Engine/Plugins/Animation/BlendStack/Source/Runtime/Public/BlendStack/AnimNode_BlendStack.h
+++ b/Engine/Plugins/Animation/BlendStack/Source/Runtime/Public/BlendStack/AnimNode_BlendStack.h
@@ -8,6 +8,17 @@
 #include "AnimNodes/AnimNode_Mirror.h"
 #include "AnimNode_BlendStack.generated.h"
 
+UENUM()
+enum class EBlendStack_BlendspaceUpdateMode : uint8
+{
+	// Only update the blendspace xy inputs once on blend in.
+	InitialOnly,
+	// Update the active/most recent blendspace xy inputs every frame.
+	UpdateActiveOnly,
+	// Update all blendspaces xy inputs every frame.
+	UpdateAll,
+};
+
 USTRUCT()
 struct BLENDSTACK_API FBlendStackAnimPlayer
 {
@@ -168,6 +179,9 @@ protected:
 		bool bUseInertialBlend, const FVector& BlendParameters, float PlayRate, float ActivationDelay,
 		FName GroupName, EAnimGroupRole::Type GroupRole, EAnimSyncMethod Method, bool bOverridePositionWhenJoiningSyncGroupAsLeader);
 
+	// Call this to update current blendspace player xy blend parameters. By default, we only update them on initial BlendTo.
+	void UpdateBlendspaceParameters(const EBlendStack_BlendspaceUpdateMode UpdateMode, const FVector& BlendParameters);
+
 	static void BlendWithPose(FAnimationPoseData& InOutPoseData, const FAnimationPoseData& OtherPoseData, const float InOutPoseWeight);
 	static void BlendWithPosePerBone(FAnimationPoseData& InOutPoseData, const FAnimationPoseData& OtherPoseData, TConstArrayView<float> OtherPoseWeights);
 
@@ -251,14 +265,23 @@ struct BLENDSTACK_API FAnimNode_BlendStack : public FAnimNode_BlendStack_Standal
 	UPROPERTY(EditAnywhere, BlueprintReadWrite, Category = Settings, meta = (PinHiddenByDefault))
 	EAlphaBlendOption BlendOption = EAlphaBlendOption::Linear;
 
+	// How we should update individual blend space parameters. See dropdown options tooltips.
+	UPROPERTY(EditAnywhere, Category = Blendspace)
+	EBlendStack_BlendspaceUpdateMode BlendspaceUpdateMode = EBlendStack_BlendspaceUpdateMode::InitialOnly;
+
 	// requested blend space blend parameters (if AnimationAsset is a blend space)
-	UPROPERTY(EditAnywhere, BlueprintReadWrite, Category = Settings, meta = (PinHiddenByDefault))
+	UPROPERTY(EditAnywhere, BlueprintReadWrite, Category = Blendspace, meta = (PinHiddenByDefault))
 	FVector BlendParameters = FVector::Zero();
 
 	// if set and bMirrored MirrorDataTable will be used for mirroring the aniamtion
 	UPROPERTY(EditAnywhere, Category = Settings, meta = (PinHiddenByDefault))
 	TObjectPtr<UMirrorDataTable> MirrorDataTable;
 
+	// Use this to define a threshold to trigger a new blend when blendspace xy input pins change.
+	// By default, any delta will trigger a blend.
+	UPROPERTY(EditAnywhere, Category = Blendspace)
+	float BlendParametersDeltaThreshold = 0.0f;
+
 	// tunable animation transition blend time 
 	UPROPERTY(EditAnywhere, Category = Settings)
 	bool bUseInertialBlend = false;
diff --git a/Engine/Plugins/Experimental/Chooser/Source/Chooser/Private/AnimNode_ChooserPlayer.cpp b/Engine/Plugins/Experimental/Chooser/Source/Chooser/Private/AnimNode_ChooserPlayer.cpp
--- a/Engine/Plugins/Experimental/Chooser/Source/Chooser/Private/AnimNode_ChooserPlayer.cpp
+++ b/Engine/Plugins/Experimental/Chooser/Source/Chooser/Private/AnimNode_ChooserPlayer.cpp
@@ -185,22 +185,11 @@ void FAnimNode_ChooserPlayer::UpdateAssetPlayer(const FAnimationUpdateContext& C
 
 	// Update blend space parameters
 	const FVector BlendParameters(BlendSpaceX, BlendSpaceY, 0);
-	if (bUpdateAllActiveBlendSpaces)
-	{
-		// apply blend space parameters to all blendspaces that are playing, including ones that are blending out
-		for (FBlendStackAnimPlayer& Player : AnimPlayers)
-		{
-			Player.SetBlendParameters(BlendParameters);
-		}
-	}
-	else
-	{
-		// apply blend space parameters only to the blendspace that is playing/blending in
-		if (!AnimPlayers.IsEmpty())
-		{
-			AnimPlayers[0].SetBlendParameters(BlendParameters);
-		}
-	}
+	const EBlendStack_BlendspaceUpdateMode BlendspaceUpdateMode = bUpdateAllActiveBlendSpaces ? 
+																EBlendStack_BlendspaceUpdateMode::UpdateAll :
+																EBlendStack_BlendspaceUpdateMode::UpdateActiveOnly;
+
+	UpdateBlendspaceParameters(BlendspaceUpdateMode, BlendParameters);
 
 	FAnimNode_BlendStack_Standalone::UpdateAssetPlayer(Context);
 
