Blend Stack: Fix for GetRelevant family of state machine functions

diff --git a/Engine/Plugins/Animation/BlendStack/Source/Runtime/Private/AnimNode_BlendStack.cpp b/Engine/Plugins/Animation/BlendStack/Source/Runtime/Private/AnimNode_BlendStack.cpp
--- a/Engine/Plugins/Animation/BlendStack/Source/Runtime/Private/AnimNode_BlendStack.cpp
+++ b/Engine/Plugins/Animation/BlendStack/Source/Runtime/Private/AnimNode_BlendStack.cpp
@@ -281,6 +281,36 @@ float FBlendStackAnimPlayer::GetAccumulatedTime() const
 	return 0.f;
 }
 
+float FBlendStackAnimPlayer::GetCurrentAssetTime() const
+{
+	if (SequencePlayerNode.GetSequence())
+	{
+		return SequencePlayerNode.GetCurrentAssetTime();
+	}
+
+	if (BlendSpacePlayerNode.GetBlendSpace())
+	{
+		return BlendSpacePlayerNode.GetCurrentAssetTime();
+	}
+
+	return 0.f;
+}
+
+float FBlendStackAnimPlayer::GetCurrentAssetLength() const
+{
+	if (SequencePlayerNode.GetSequence())
+	{
+		return SequencePlayerNode.GetCurrentAssetLength();
+	}
+
+	if (BlendSpacePlayerNode.GetBlendSpace())
+	{
+		return BlendSpacePlayerNode.GetCurrentAssetLength();
+	}
+
+	return 0.f;
+}
+
 float FBlendStackAnimPlayer::GetPlayRate() const
 {
 	if (SequencePlayerNode.GetSequence())
@@ -829,6 +859,16 @@ void FAnimNode_BlendStack_Standalone::InitializeSample(const FAnimationInitializ
 	}
 }
 
+float FAnimNode_BlendStack_Standalone::GetCurrentAssetLength() const
+{
+	return AnimPlayers.IsEmpty() ? 0.0f : AnimPlayers[0].GetCurrentAssetLength();
+}
+
+float FAnimNode_BlendStack_Standalone::GetCurrentAssetTime() const
+{
+	return AnimPlayers.IsEmpty() ? 0.0f : AnimPlayers[0].GetCurrentAssetTime();
+}
+
 UAnimationAsset* FAnimNode_BlendStack_Standalone::GetAnimAsset() const
 {
 	return AnimPlayers.IsEmpty() ? nullptr : AnimPlayers[0].GetAnimationAsset();
diff --git a/Engine/Plugins/Animation/BlendStack/Source/Runtime/Public/BlendStack/AnimNode_BlendStack.h b/Engine/Plugins/Animation/BlendStack/Source/Runtime/Public/BlendStack/AnimNode_BlendStack.h
--- a/Engine/Plugins/Animation/BlendStack/Source/Runtime/Public/BlendStack/AnimNode_BlendStack.h
+++ b/Engine/Plugins/Animation/BlendStack/Source/Runtime/Public/BlendStack/AnimNode_BlendStack.h
@@ -32,6 +32,8 @@ struct BLENDSTACK_API FBlendStackAnimPlayer
 	void Evaluate_AnyThread(FPoseContext& Output);
 	void Update_AnyThread(const FAnimationUpdateContext& Context);
 	float GetAccumulatedTime() const;
+	float GetCurrentAssetTime() const;
+	float GetCurrentAssetLength() const;
 	float GetPlayRate() const;
 	
 	float GetBlendInPercentage() const;
@@ -165,6 +167,8 @@ struct BLENDSTACK_API FAnimNode_BlendStack_Standalone : public FAnimNode_AssetPl
 	virtual UAnimationAsset* GetAnimAsset() const;
 	virtual float GetAccumulatedTime() const override;
 	virtual void UpdateAssetPlayer(const FAnimationUpdateContext& Context) override;
+	virtual float GetCurrentAssetTime() const override;
+	virtual float GetCurrentAssetLength() const override;
 	// End of FAnimNode_AssetPlayerBase interface
 
 	int32 GetNextPoseLinkIndex();
