pose search
- Blend space playback crash fix
- fix for Missing discard flag in PoseSearch debugger

diff --git a/Engine/Plugins/Animation/BlendStack/Source/Runtime/Private/AnimNode_BlendStack.cpp b/Engine/Plugins/Animation/BlendStack/Source/Runtime/Private/AnimNode_BlendStack.cpp
--- a/Engine/Plugins/Animation/BlendStack/Source/Runtime/Private/AnimNode_BlendStack.cpp
+++ b/Engine/Plugins/Animation/BlendStack/Source/Runtime/Private/AnimNode_BlendStack.cpp
@@ -71,6 +71,8 @@ void FBlendStackAnimPlayer::Initialize(const FAnimationInitializeContext& Contex
 	}
 	else if (UAnimSequenceBase* SequenceBase = Cast<UAnimSequenceBase>(AnimationAsset))
 	{
+		BlendSpacePlayerNode.SetBlendSpace(nullptr);
+
 		SequencePlayerNode.SetAccumulatedTime(AccumulatedTime);
 		SequencePlayerNode.SetSequence(SequenceBase);
 		SequencePlayerNode.SetLoopAnimation(bLoop);
@@ -82,6 +84,8 @@ void FBlendStackAnimPlayer::Initialize(const FAnimationInitializeContext& Contex
 	}
 	else if (UBlendSpace* BlendSpace = Cast<UBlendSpace>(AnimationAsset))
 	{
+		SequencePlayerNode.SetSequence(nullptr);
+
 		// making sure AccumulatedTime is in normalized space
 		AccumulatedTime = FMath::Clamp(AccumulatedTime, 0.f, 1.f);
 
@@ -103,6 +107,9 @@ void FBlendStackAnimPlayer::Initialize(const FAnimationInitializeContext& Contex
 
 	if (bUnsupportedAnimAsset)
 	{
+		BlendSpacePlayerNode.SetBlendSpace(nullptr);
+		SequencePlayerNode.SetSequence(nullptr);
+
 		UE_LOG(LogBlendStack, Error, TEXT("FBlendStackAnimPlayer unsupported AnimationAsset %s"), *GetNameSafe(AnimationAsset));
 	}
 
diff --git a/Engine/Plugins/Animation/PoseSearch/Source/Editor/Private/PoseSearchDebuggerDatabaseColumns.h b/Engine/Plugins/Animation/PoseSearch/Source/Editor/Private/PoseSearchDebuggerDatabaseColumns.h
--- a/Engine/Plugins/Animation/PoseSearch/Source/Editor/Private/PoseSearchDebuggerDatabaseColumns.h
+++ b/Engine/Plugins/Animation/PoseSearch/Source/Editor/Private/PoseSearchDebuggerDatabaseColumns.h
@@ -641,6 +641,11 @@ struct FPoseCandidateFlags : ITextColumn
 				Sring.Append("S ");
 			}
 
+			if (EnumHasAnyFlags(Row->PoseCandidateFlags, EPoseCandidateFlags::DiscardedBy_AssetReselection))
+			{
+				Sring.Append("R ");
+			}
+
 			return FText::FromString(Sring);
 		}
 
@@ -684,6 +689,11 @@ struct FPoseCandidateFlags : ITextColumn
 				TextBuilder.AppendLine(LOCTEXT("DiscardedBy_Search_Tooltip", "(S) Search"));
 			}
 
+			if (EnumHasAnyFlags(Row->PoseCandidateFlags, EPoseCandidateFlags::DiscardedBy_AssetReselection))
+			{
+				TextBuilder.AppendLine(LOCTEXT("DiscardedBy_AssetReselection_Tooltip", "(R) Disable Reselection"));
+			}
+
 			return TextBuilder.ToText();
 		}
 
diff --git a/Engine/Plugins/Animation/PoseSearch/Source/Runtime/Private/PoseSearchDatabase.cpp b/Engine/Plugins/Animation/PoseSearch/Source/Runtime/Private/PoseSearchDatabase.cpp
--- a/Engine/Plugins/Animation/PoseSearch/Source/Runtime/Private/PoseSearchDatabase.cpp
+++ b/Engine/Plugins/Animation/PoseSearch/Source/Runtime/Private/PoseSearchDatabase.cpp
@@ -1134,7 +1134,7 @@ void UPoseSearchDatabase::PopulateNonSelectableIdx(FNonSelectableIdx& NonSelecta
 #if UE_POSE_SEARCH_TRACE_ENABLED
 						const TArray<float> PoseValues = SearchIndex.GetPoseValuesSafe(PoseIdx);
 						const FPoseSearchCost PoseCost = SearchIndex.ComparePoses(PoseIdx, 0.f, PoseValues, QueryValues);
-						SearchContext.Track(this, PoseIdx, EPoseCandidateFlags::DiscardedBy_PoseJumpThresholdTime, PoseCost);
+						SearchContext.Track(this, PoseIdx, EPoseCandidateFlags::DiscardedBy_AssetReselection, PoseCost);
 #endif // UE_POSE_SEARCH_TRACE_ENABLED
 					}
 				}
diff --git a/Engine/Plugins/Animation/PoseSearch/Source/Runtime/Public/PoseSearch/PoseSearchContext.h b/Engine/Plugins/Animation/PoseSearch/Source/Runtime/Public/PoseSearch/PoseSearchContext.h
--- a/Engine/Plugins/Animation/PoseSearch/Source/Runtime/Public/PoseSearch/PoseSearchContext.h
+++ b/Engine/Plugins/Animation/PoseSearch/Source/Runtime/Public/PoseSearch/PoseSearchContext.h
@@ -44,8 +44,9 @@ enum class EPoseCandidateFlags : uint32
 	DiscardedBy_PoseFilter = 1 << 6,
 	DiscardedBy_AssetIdxFilter = 1 << 7,
 	DiscardedBy_Search = 1 << 8,
+	DiscardedBy_AssetReselection = 1 << 9,
 
-	AnyDiscardedMask = DiscardedBy_PoseJumpThresholdTime | DiscardedBy_PoseReselectHistory | DiscardedBy_BlockTransition | DiscardedBy_PoseFilter | DiscardedBy_AssetIdxFilter | DiscardedBy_Search,
+	AnyDiscardedMask = DiscardedBy_PoseJumpThresholdTime | DiscardedBy_PoseReselectHistory | DiscardedBy_BlockTransition | DiscardedBy_PoseFilter | DiscardedBy_AssetIdxFilter | DiscardedBy_Search | DiscardedBy_AssetReselection,
 };
 ENUM_CLASS_FLAGS(EPoseCandidateFlags);
 
