Blend Stack: Added anim node relevancy base anim node-style functions + blend to w/ blend profile exposed to blueprints

diff --git a/Engine/Plugins/Animation/BlendStack/Source/Runtime/Private/BlendStackAnimNodeLibrary.cpp b/Engine/Plugins/Animation/BlendStack/Source/Runtime/Private/BlendStackAnimNodeLibrary.cpp
--- a/Engine/Plugins/Animation/BlendStack/Source/Runtime/Private/BlendStackAnimNodeLibrary.cpp
+++ b/Engine/Plugins/Animation/BlendStack/Source/Runtime/Private/BlendStackAnimNodeLibrary.cpp
@@ -79,6 +79,98 @@ void UBlendStackAnimNodeLibrary::BlendTo(const FAnimUpdateContext& Context,
 	}
 }
 
+void UBlendStackAnimNodeLibrary::BlendToWithSettings(const FAnimUpdateContext& Context,
+													const FBlendStackAnimNodeReference& BlendStackNode, 
+													UAnimationAsset* AnimationAsset /*= nullptr*/, 
+													float AnimationTime /*= 0.f*/, 
+													bool bLoop /*= false*/, 
+													bool bMirrored /*= false*/, 
+													float BlendTime /*= 0.2f*/, 
+													UBlendProfile* BlendProfile /*= nullptr*/, 
+													EAlphaBlendOption BlendOption /*= EAlphaBlendOption::HermiteCubic*/,
+													bool bInertialBlend /*= false*/, 
+													FVector BlendParameters /*= FVector::ZeroVector*/, 
+													float WantedPlayRate /*= 1.f*/, 
+													float ActivationDelay /*= 0.f*/)
+{
+	if (AnimationAsset != nullptr)
+	{
+		if (FAnimNode_BlendStack* BlendStackNodePtr = BlendStackNode.GetAnimNodePtr<FAnimNode_BlendStack>())
+		{
+			if (const FAnimationUpdateContext* AnimationUpdateContext = Context.GetContext())
+			{
+				BlendStackNodePtr->BlendTo(
+					*AnimationUpdateContext,
+					AnimationAsset,
+					AnimationTime,
+					bLoop,
+					bMirrored,
+					BlendStackNodePtr->MirrorDataTable,
+					BlendTime,
+					BlendProfile,
+					BlendOption,
+					bInertialBlend,
+					BlendParameters,
+					WantedPlayRate,
+					ActivationDelay,
+					BlendStackNodePtr->GetGroupName(), BlendStackNodePtr->GetGroupRole(), BlendStackNodePtr->GetGroupMethod());
+			}
+			else
+			{
+				UE_LOG(LogBlendStack, Warning, TEXT("UBlendStackAnimNodeLibrary::BlendToWithSettings called with an invalid context."));
+			}
+		}
+		else
+		{
+			UE_LOG(LogBlendStack, Warning, TEXT("UBlendStackAnimNodeLibrary::BlendToWithSettings called with an invalid type."));
+		}
+	}
+}
+
+UAnimationAsset* UBlendStackAnimNodeLibrary::GetCurrentAsset(const FBlendStackAnimNodeReference& BlendStackNode)
+{
+	if (FAnimNode_BlendStack* BlendStackNodePtr = BlendStackNode.GetAnimNodePtr<FAnimNode_BlendStack>())
+	{
+		return BlendStackNodePtr->GetAnimAsset();
+	}
+
+	UE_LOG(LogBlendStack, Warning, TEXT("UBlendStackAnimNodeLibrary::GetCurrentAsset called with an invalid type."));
+	return nullptr;
+}
+
+float UBlendStackAnimNodeLibrary::GetCurrentAssetTime(const FBlendStackAnimNodeReference& BlendStackNode)
+{
+	if (FAnimNode_BlendStack* BlendStackNodePtr = BlendStackNode.GetAnimNodePtr<FAnimNode_BlendStack>())
+	{
+		return BlendStackNodePtr->GetCurrentAssetTime();
+	}
+
+	UE_LOG(LogBlendStack, Warning, TEXT("UBlendStackAnimNodeLibrary::GetCurrentAssetTime called with an invalid type."));
+	return 0.0f;
+}
+
+float UBlendStackAnimNodeLibrary::GetCurrentAssetTimeRemaining(const FBlendStackAnimNodeReference& BlendStackNode)
+{
+	if (FAnimNode_BlendStack* BlendStackNodePtr = BlendStackNode.GetAnimNodePtr<FAnimNode_BlendStack>())
+	{
+		return BlendStackNodePtr->GetCurrentAssetLength() - BlendStackNodePtr->GetCurrentAssetTimePlayRateAdjusted();
+	}
+
+	UE_LOG(LogBlendStack, Warning, TEXT("UBlendStackAnimNodeLibrary::GetCurrentAssetTimeRemaining called with an invalid type."));
+	return 0.0f;
+}
+
+bool UBlendStackAnimNodeLibrary::IsCurrentAssetLooping(const FBlendStackAnimNodeReference& BlendStackNode)
+{
+	if (FAnimNode_BlendStack* BlendStackNodePtr = BlendStackNode.GetAnimNodePtr<FAnimNode_BlendStack>())
+	{
+		return BlendStackNodePtr->IsLooping();
+	}
+
+	UE_LOG(LogBlendStack, Warning, TEXT("UBlendStackAnimNodeLibrary::IsCurrentAssetLooping called with an invalid type."));
+	return false;
+}
+
 void UBlendStackAnimNodeLibrary::ForceBlendNextUpdate(const FBlendStackAnimNodeReference& BlendStackNode)
 {
 	if (FAnimNode_BlendStack* BlendStackNodePtr = BlendStackNode.GetAnimNodePtr<FAnimNode_BlendStack>())
diff --git a/Engine/Plugins/Animation/BlendStack/Source/Runtime/Public/BlendStack/BlendStackAnimNodeLibrary.h b/Engine/Plugins/Animation/BlendStack/Source/Runtime/Public/BlendStack/BlendStackAnimNodeLibrary.h
--- a/Engine/Plugins/Animation/BlendStack/Source/Runtime/Public/BlendStack/BlendStackAnimNodeLibrary.h
+++ b/Engine/Plugins/Animation/BlendStack/Source/Runtime/Public/BlendStack/BlendStackAnimNodeLibrary.h
@@ -5,9 +5,11 @@
 #include "Animation/AnimExecutionContext.h"
 #include "Animation/AnimNodeReference.h"
 #include "Kismet/BlueprintFunctionLibrary.h"
+#include "AlphaBlend.h"
 #include "BlendStackAnimNodeLibrary.generated.h"
 
 struct FAnimNode_BlendStack;
+class UBlendProfile;
 
 USTRUCT(BlueprintType)
 struct FBlendStackAnimNodeReference : public FAnimNodeReference
@@ -60,4 +62,32 @@ public:
 						FVector BlendParameters = FVector::ZeroVector,
 						float WantedPlayRate = 1.f,
 						float ActivationDelay = 0.f);
+
+	// Note: Experimental and subject to change!
+	UFUNCTION(BlueprintCallable, Category = "Animation|BlendStack|Experimental", meta = (BlueprintThreadSafe, DisplayName = "Blend To"))
+	static void BlendToWithSettings(const FAnimUpdateContext& Context, 
+						const FBlendStackAnimNodeReference& BlendStackNode, 
+						UAnimationAsset* AnimationAsset = nullptr,
+						float AnimationTime = 0.f,
+						bool bLoop = false,
+						bool bMirrored = false,
+						float BlendTime = 0.2f,
+						UBlendProfile* BlendProfile = nullptr,
+						EAlphaBlendOption BlendOption = EAlphaBlendOption::HermiteCubic,
+						bool bInertialBlend = false,
+						FVector BlendParameters = FVector::ZeroVector,
+						float WantedPlayRate = 1.f,
+						float ActivationDelay = 0.f);
+
+	UFUNCTION(BlueprintPure, Category = "Animation|BlendStack", meta = (BlueprintThreadSafe))
+    static UAnimationAsset* GetCurrentAsset(const FBlendStackAnimNodeReference& BlendStackNode);
+
+	UFUNCTION(BlueprintPure, Category = "Animation|BlendStack", meta = (BlueprintThreadSafe))
+    static float GetCurrentAssetTime(const FBlendStackAnimNodeReference& BlendStackNode);
+
+	UFUNCTION(BlueprintPure, Category = "Animation|BlendStack", meta = (BlueprintThreadSafe))
+    static float GetCurrentAssetTimeRemaining(const FBlendStackAnimNodeReference& BlendStackNode);
+
+	UFUNCTION(BlueprintPure, Category = "Animation|BlendStack", meta = (BlueprintThreadSafe))
+    static bool IsCurrentAssetLooping(const FBlendStackAnimNodeReference& BlendStackNode);
 };
\ No newline at end of file
