Fix build NonUnity missing include BlendStack

#tests nonunity

diff --git a/Engine/Plugins/Animation/BlendStack/Source/Runtime/Public/BlendStack/AnimNode_BlendStack.h b/Engine/Plugins/Animation/BlendStack/Source/Runtime/Public/BlendStack/AnimNode_BlendStack.h
--- a/Engine/Plugins/Animation/BlendStack/Source/Runtime/Public/BlendStack/AnimNode_BlendStack.h
+++ b/Engine/Plugins/Animation/BlendStack/Source/Runtime/Public/BlendStack/AnimNode_BlendStack.h
@@ -6,6 +6,7 @@
 #include "Animation/AnimNode_SequencePlayer.h"
 #include "AnimNodes/AnimNode_BlendSpacePlayer.h"
 #include "AnimNodes/AnimNode_Mirror.h"
+#include "EngineDefines.h"
 #include "AnimNode_BlendStack.generated.h"
 
 UENUM()
