Add guard objectID for cooker builds

#tests nonunity

diff --git a/Engine/Plugins/Animation/BlendStack/Source/Runtime/Private/AnimNode_BlendStack.cpp b/Engine/Plugins/Animation/BlendStack/Source/Runtime/Private/AnimNode_BlendStack.cpp
--- a/Engine/Plugins/Animation/BlendStack/Source/Runtime/Private/AnimNode_BlendStack.cpp
+++ b/Engine/Plugins/Animation/BlendStack/Source/Runtime/Private/AnimNode_BlendStack.cpp
@@ -119,7 +119,7 @@ void FBlendStackAnimPlayer::Initialize(const FAnimationInitializeContext& Contex
 
 	OverrideCurve.Empty();
 
-#if ENABLE_ANIM_DEBUG || ENABLE_VISUAL_LOG 
+#if OBJECT_TRACE_ENABLED && (ENABLE_ANIM_DEBUG || ENABLE_VISUAL_LOG)
 	// Matches 'MakeBlendWeightCurveColor' in DebugWeightsTrack.cpp to match rewind debugger color (Does not depend since do not want dependency on debug for color).
 	auto MakeDebugColor = [](uint32 InSeed)
 	{
