Blend profiles now work correctly with Blend Stack and Motion Matching

#tests win64 editor animsandbox
[RN] fix Animation

diff --git a/Engine/Plugins/Animation/BlendStack/Source/Runtime/Private/AnimNode_BlendStack.cpp b/Engine/Plugins/Animation/BlendStack/Source/Runtime/Private/AnimNode_BlendStack.cpp
--- a/Engine/Plugins/Animation/BlendStack/Source/Runtime/Private/AnimNode_BlendStack.cpp
+++ b/Engine/Plugins/Animation/BlendStack/Source/Runtime/Private/AnimNode_BlendStack.cpp
@@ -504,6 +504,33 @@ void FBlendStackAnimPlayer::GetBlendInWeights(TArrayView<float> Weights) const
 	}
 }
 
+// This method fills weights indexed by the compact skeleton indices
+// Todo: move to struct method after 5.5. 
+// This needs to be a free function for a hot fix to avoid breaking ABI
+void GetBlendInWeightsCompactPose(const FBlendStackAnimPlayer& AnimPlayer, const TCustomBoneIndexArray<float, FSkeletonPoseBoneIndex>& TotalBlendInTimePerBone, TArrayView<float> Weights, const FBoneContainer& BoneContainer)
+{	
+	const int32 CompactPoseNumBones = BoneContainer.GetCompactPoseNumBones();
+	check(Weights.Num() == CompactPoseNumBones);
+
+	const float WeightForZeroBlendInTime = AnimPlayer.GetTimeToActivation() > 0.f ? 0.f : 1.f;
+	for (int32 CompactPoseBoneIdx = 0; CompactPoseBoneIdx < CompactPoseNumBones; ++CompactPoseBoneIdx)
+	{
+		const int32 SkeletonBoneIdx = BoneContainer.GetSkeletonIndex(FCompactPoseBoneIndex(CompactPoseBoneIdx));
+		check(SkeletonBoneIdx != INDEX_NONE);
+		const float TotalBlendInTimeBoneIdx = TotalBlendInTimePerBone[SkeletonBoneIdx];
+		if (TotalBlendInTimeBoneIdx < UE_SMALL_NUMBER)
+		{
+			Weights[CompactPoseBoneIdx] = WeightForZeroBlendInTime;
+		}
+		else
+		{
+			check(AnimPlayer.GetCurrentBlendInTime() >= 0.f);
+			const float LinearWeight = FMath::Min(AnimPlayer.GetCurrentBlendInTime() / TotalBlendInTimeBoneIdx, 1.f);
+			Weights[CompactPoseBoneIdx] = FAlphaBlend::AlphaToBlendOption(LinearWeight, AnimPlayer.GetBlendOption());
+		}
+	}
+}
+
 /////////////////////////////////////////////////////
 // FAnimNode_BlendStack_Standalone
 
@@ -613,13 +640,10 @@ void FAnimNode_BlendStack_Standalone::Evaluate_AnyThread(FPoseContext& Output)
 
 		FPoseContext EvaluationPoseContext(Output);
 
-		const USkeleton* SkeletonAsset = Output.AnimInstanceProxy->GetRequiredBones().GetSkeletonAsset();
-		check(SkeletonAsset);
-
-		const FReferenceSkeleton& RefSkeleton = SkeletonAsset->GetReferenceSkeleton();
-		const int32 NumSkeletonBones = RefSkeleton.GetNum();
+		const FBoneContainer& BoneContainer = Output.Pose.GetBoneContainer();
+		const int32 NumCompactPoseBones = BoneContainer.GetCompactPoseNumBones();
 
-		auto EvaluateAndBlendPlayerByIndex = [this, &EvaluationPoseContext, &Output](int32 PlayerIndex)
+		auto EvaluateAndBlendPlayerByIndex = [this, &EvaluationPoseContext, &Output, &BoneContainer, NumCompactPoseBones](int32 PlayerIndex)
 		{
 			FAnimationPoseData OutputAnimationPoseData(Output);
 			FAnimationPoseData EvaluationAnimationPoseData(EvaluationPoseContext);
@@ -630,8 +654,8 @@ void FAnimNode_BlendStack_Standalone::Evaluate_AnyThread(FPoseContext& Output)
 				const int32 BlendInWeightsNum = AnimPlayers[PlayerIndex].GetBlendInWeightsNum();
 				if (BlendInWeightsNum > 0)
 				{
-					TArrayView<float> EvaluationAnimationPoseDataWeights((float*)FMemory_Alloca(BlendInWeightsNum * sizeof(float)), BlendInWeightsNum);
-					AnimPlayers[PlayerIndex].GetBlendInWeights(EvaluationAnimationPoseDataWeights);
+					TArrayView<float> EvaluationAnimationPoseDataWeights((float*)FMemory_Alloca(NumCompactPoseBones * sizeof(float)), NumCompactPoseBones);
+					GetBlendInWeightsCompactPose(AnimPlayers[PlayerIndex], AnimPlayers[PlayerIndex].TotalBlendInTimePerBone, EvaluationAnimationPoseDataWeights, BoneContainer);
 					BlendWithPosePerBone(OutputAnimationPoseData, EvaluationAnimationPoseData, EvaluationAnimationPoseDataWeights);
 				}
 				else
diff --git a/Engine/Plugins/Animation/BlendStack/Source/Runtime/Public/BlendStack/AnimNode_BlendStack.h b/Engine/Plugins/Animation/BlendStack/Source/Runtime/Public/BlendStack/AnimNode_BlendStack.h
--- a/Engine/Plugins/Animation/BlendStack/Source/Runtime/Public/BlendStack/AnimNode_BlendStack.h
+++ b/Engine/Plugins/Animation/BlendStack/Source/Runtime/Public/BlendStack/AnimNode_BlendStack.h
@@ -24,6 +24,7 @@ USTRUCT()
 struct BLENDSTACK_API FBlendStackAnimPlayer
 {
 	GENERATED_BODY()
+	friend struct FAnimNode_BlendStack_Standalone;
 	
 	void Initialize(const FAnimationInitializeContext& Context, UAnimationAsset* AnimationAsset, float AccumulatedTime, bool bLoop,
 		bool bMirrored, UMirrorDataTable* MirrorDataTable, float BlendTime, const UBlendProfile* BlendProfile, EAlphaBlendOption InBlendOption,
@@ -40,6 +41,7 @@ struct BLENDSTACK_API FBlendStackAnimPlayer
 	float GetBlendInPercentage() const;
 	
 	int32 GetBlendInWeightsNum() const;
+	// Important: this method fills weight array indexed with the full skeleton indices, not the compact pose!
 	void GetBlendInWeights(TArrayView<float> Weights) const;
 	float GetBlendInWeight() const;
 
