blendstack - curve blending fix

diff --git a/Engine/Plugins/Animation/BlendStack/Source/Runtime/Private/AnimNode_BlendStack.cpp b/Engine/Plugins/Animation/BlendStack/Source/Runtime/Private/AnimNode_BlendStack.cpp
--- a/Engine/Plugins/Animation/BlendStack/Source/Runtime/Private/AnimNode_BlendStack.cpp
+++ b/Engine/Plugins/Animation/BlendStack/Source/Runtime/Private/AnimNode_BlendStack.cpp
@@ -2,8 +2,9 @@
 
 #include "BlendStack/AnimNode_BlendStack.h"
 #include "Algo/MaxElement.h"
-#include "Animation/AnimComposite.h"
 #include "Animation/AnimBlendDebugScope.h"
+#include "Animation/AnimComposite.h"
+#include "Animation/AnimInertializationSyncScope.h"
 #include "Animation/AnimInstanceProxy.h"
 #include "Animation/AnimMontage.h"
 #include "Animation/AnimNode_Inertialization.h"
@@ -11,20 +12,23 @@
 #include "Animation/AnimSequence.h"
 #include "Animation/AnimTrace.h"
 #include "Animation/BlendSpace.h"
+#include "Animation/NamedValueArray.h"
+#include "BlendedCurveUtils.h"
 #include "BlendStack/AnimNode_BlendStackInput.h"
-#include "BlendStackAnimEventsFilterScope.h"
-#include "Animation/AnimInertializationSyncScope.h"
 #include "BlendStack/BlendStackDefines.h"
+#include "BlendStackAnimEventsFilterScope.h"
 
 #include UE_INLINE_GENERATED_CPP_BY_NAME(AnimNode_BlendStack)
 
-#if ENABLE_ANIM_DEBUG
-static TAutoConsoleVariable<bool> CVarAnimBlendStackEnable(TEXT("a.AnimNode.BlendStack.Enable"), true, TEXT("Enable / Disable Blend Stack"));
-#endif
-
 #define LOCTEXT_NAMESPACE "AnimNode_BlendStack"
 
-static constexpr FBoneIndexType RootBoneIndexType = 0;
+namespace UE::BlendStack
+{
+#if ENABLE_ANIM_DEBUG
+	static bool GVarAnimBlendStackEnable = true;
+	static FAutoConsoleVariableRef CVarAnimBlendStackEnable(TEXT("a.AnimNode.BlendStack.Enable"), GVarAnimBlendStackEnable, TEXT("Enable / Disable Blend Stack"));
+#endif
+} // namespace UE::BlendStack
 
 /////////////////////////////////////////////////////
 // FBlendStackAnimPlayer
@@ -586,7 +590,7 @@ void FAnimNode_BlendStack_Standalone::Evaluate_AnyThread(FPoseContext& Output)
 
 	bool bDisableBlendStack = false;
 #if ENABLE_ANIM_DEBUG
-	bDisableBlendStack = !CVarAnimBlendStackEnable.GetValueOnAnyThread();
+	bDisableBlendStack = !UE::BlendStack::GVarAnimBlendStackEnable;
 
 	if (UE_TRACE_CHANNELEXPR_IS_ENABLED(AnimationChannel))
 	{
@@ -942,7 +946,7 @@ float FAnimNode_BlendStack_Standalone::GetAccumulatedTime() const
 static void RequestInertialBlend(const FAnimationUpdateContext& Context, float BlendTime, const UBlendProfile* BlendProfile, EAlphaBlendOption BlendOption)
 {
 #if ENABLE_ANIM_DEBUG
-	if (!CVarAnimBlendStackEnable.GetValueOnAnyThread())
+	if (!UE::BlendStack::GVarAnimBlendStackEnable)
 	{
 		return;
 	}
@@ -1153,6 +1157,7 @@ void FAnimNode_BlendStack_Standalone::GatherDebugData(FNodeDebugData& DebugData)
 void FAnimNode_BlendStack_Standalone::BlendWithPosePerBone(FAnimationPoseData& InOutPoseData, const FAnimationPoseData& OtherPoseData, TConstArrayView<float> OtherPoseWeights)
 {
 	using namespace UE::Anim;
+	using namespace UE::BlendStack;
 
 	FCompactPose& InOutPose = InOutPoseData.GetPose();
 	FBlendedCurve& InOutCurve = InOutPoseData.GetCurve();
@@ -1181,11 +1186,7 @@ void FAnimNode_BlendStack_Standalone::BlendWithPosePerBone(FAnimationPoseData& I
 	// Ensure that all of the resulting rotations are normalized
 	InOutPose.NormalizeRotations();
 
-	// @note : This isn't perfect as curve can link to joint, and it would be the best to use that information
-	// but that is very expensive option as we have to have another indirect look up table to search. 
-	// For now, replacing with combine (non-zero will be overriden)
-	// in the future, we might want to do this outside if we want per bone blend to apply curve also UE-39182
-	InOutPoseData.GetCurve().Combine(OtherPoseData.GetCurve());
+	FBlendedCurveUtils::LerpToPerBoneEx(InOutCurve, OtherPoseData.GetCurve(), InOutPoseData.GetPose().GetBoneContainer(), OtherPoseWeights);
 
 	// @todo: optimize away the copy
 	FStackAttributeContainer CustomAttributes;
@@ -1198,6 +1199,7 @@ void FAnimNode_BlendStack_Standalone::BlendWithPosePerBone(FAnimationPoseData& I
 void FAnimNode_BlendStack_Standalone::BlendWithPose(FAnimationPoseData& InOutPoseData, const FAnimationPoseData& OtherPoseData, const float InOutPoseWeight)
 {
 	using namespace UE::Anim;
+	using namespace UE::BlendStack;
 
 	FCompactPose& InOutPose = InOutPoseData.GetPose();
 	FBlendedCurve& InOutCurve = InOutPoseData.GetCurve();
@@ -1220,7 +1222,7 @@ void FAnimNode_BlendStack_Standalone::BlendWithPose(FAnimationPoseData& InOutPos
 	// Ensure that all of the resulting rotations are normalized
 	InOutPose.NormalizeRotations();
 
-	InOutCurve.LerpTo(OtherPoseData.GetCurve(), OtherPoseWeight);
+	FBlendedCurveUtils::LerpToEx(InOutCurve, OtherPoseData.GetCurve(), OtherPoseWeight);
 	
 	// @todo: optimize away the copy
 	FStackAttributeContainer CustomAttributes;
diff --git a/Engine/Plugins/Animation/BlendStack/Source/Runtime/Private/BlendedCurveUtils.cpp b/Engine/Plugins/Animation/BlendStack/Source/Runtime/Private/BlendedCurveUtils.cpp
new file mode 100644
--- /dev/null
+++ b/Engine/Plugins/Animation/BlendStack/Source/Runtime/Private/BlendedCurveUtils.cpp
@@ -0,0 +1,79 @@
+// Copyright Epic Games, Inc. All Rights Reserved.
+
+#include "BlendedCurveUtils.h"
+#include "Animation/Skeleton.h"
+#include "BoneContainer.h"
+#include "HAL/IConsoleManager.h"
+
+namespace UE::BlendStack
+{
+	
+bool GVarUseBlendCurveFixes = true;
+static FAutoConsoleVariableRef CVarUseBlendCurveFixes(TEXT("a.AnimNode.BlendStack.UseBlendCurvePerBone"), GVarUseBlendCurveFixes, TEXT("Enable BlendStack BlendCurve fixes. It causes behavioral changes. For the good"));
+
+// function adapted from TBaseBlendedCurve::LerpTo, that uses UnionEx instead of UE::Anim::FNamedValueArrayUtils::Union
+void FBlendedCurveUtils::LerpToEx(FBlendedCurve& InOutCurve, const FBlendedCurve& OtherCurve, float Alpha)
+{
+	if (!FAnimWeight::IsRelevant(FMath::Abs(Alpha)))
+	{
+	}
+	else if (!FAnimWeight::IsRelevant(FMath::Abs(Alpha - 1.0f)))
+	{
+		// if blend is all the way for child2, then just copy its bone atoms
+		InOutCurve.Override(OtherCurve);
+	}
+	else
+	{
+		// Combine using Lerp. Result is a merged set of curves in 'this'. 
+		UnionEx(InOutCurve, OtherCurve, [&Alpha](FBlendedCurve::ElementType& InOutThisElement, const FBlendedCurve::ElementType& InOtherElement, UE::Anim::ENamedValueUnionFlags InFlags)
+			{
+				InOutThisElement.Value = FMath::Lerp(InOutThisElement.Value, InOtherElement.Value, Alpha);
+				InOutThisElement.Flags |= InOtherElement.Flags;
+			});
+	}
+}
+
+// LerpToEx with per linked bone weighting blend
+void FBlendedCurveUtils::LerpToPerBoneEx(FBlendedCurve& InOutCurve, const FBlendedCurve& OtherCurve, const FBoneContainer& BoneContainer, TConstArrayView<float> OtherCurveBoneWeights)
+{
+	// NoTe: BEHAVIOR CHANGE START
+	//       as the note in the else branch states, the current behavior is not perfect, so for now we encapsulate this new behaviors behind CVar for A/B testing
+	if (GVarUseBlendCurveFixes)
+	{
+		const USkeleton* Skeleton = BoneContainer.GetSkeletonAsset();
+
+		UnionEx(InOutCurve, OtherCurve,
+			[Skeleton, &BoneContainer, &OtherCurveBoneWeights](FBlendedCurve::ElementType& InOutThisElement, const FBlendedCurve::ElementType& InOtherElement, UE::Anim::ENamedValueUnionFlags InFlags)
+			{
+				check(InOutThisElement.Name == InOtherElement.Name);
+
+				float Weight = OtherCurveBoneWeights[0];
+				if (const FCurveMetaData* CurveMetaData = Skeleton->GetCurveMetaData(InOutThisElement.Name))
+				{
+					for (const FBoneReference& LinkedBone : CurveMetaData->LinkedBones)
+					{
+						const FCompactPoseBoneIndex CompactPoseIndex = LinkedBone.GetCompactPoseIndex(BoneContainer);
+						if (CompactPoseIndex != INDEX_NONE)
+						{
+							Weight = OtherCurveBoneWeights[CompactPoseIndex.GetInt()];
+							break;
+						}
+					}
+				}
+
+				InOutThisElement.Value = FMath::Lerp(InOutThisElement.Value, InOtherElement.Value, Weight);
+				InOutThisElement.Flags |= InOtherElement.Flags;
+			});
+	}
+	// NoTe: BEHAVIOR CHANGE END
+	else
+	{
+		// @note : This isn't perfect as curve can link to joint, and it would be the best to use that information
+		// but that is very expensive option as we have to have another indirect look up table to search. 
+		// For now, replacing with combine (non-zero will be overriden)
+		// in the future, we might want to do this outside if we want per bone blend to apply curve also UE-39182
+		InOutCurve.Combine(OtherCurve);
+	}
+}
+
+} // namespace UE::BlendStack
diff --git a/Engine/Plugins/Animation/BlendStack/Source/Runtime/Private/BlendedCurveUtils.h b/Engine/Plugins/Animation/BlendStack/Source/Runtime/Private/BlendedCurveUtils.h
new file mode 100644
--- /dev/null
+++ b/Engine/Plugins/Animation/BlendStack/Source/Runtime/Private/BlendedCurveUtils.h
@@ -0,0 +1,173 @@
+// Copyright Epic Games, Inc. All Rights Reserved.
+
+#pragma once
+
+#include "Animation/AnimCurveTypes.h"
+
+namespace UE::BlendStack
+{
+
+// Experimental, this feature might be removed without warning, not for production use
+extern bool GVarUseBlendCurveFixes;
+
+// Experimental, this feature might be removed without warning, not for production use
+// Utils struct used to perform Curve blending with additiona fixes that have behavioral changes before integrating changes back to their original locations (look for code comments)
+struct FBlendedCurveUtils : FBlendedCurve
+{
+	// function adapted from UE::Anim::FNamedValueArrayUtils::Union, with additional bugfixes (read comments) that changes behaviors
+	template<typename PredicateType>
+	static void UnionEx(FBlendedCurve& BlendedCurve0, const FBlendedCurve& BlendedCurve1, PredicateType InPredicate)
+	{
+		// reinterpret_cast hack to access FBlendedCurve protected properties and methods (SortElementsIfRequired, Elements)
+		FBlendedCurveUtils& InOutValueArray0 = reinterpret_cast<FBlendedCurveUtils&>(BlendedCurve0);
+		const FBlendedCurveUtils& InValueArray1 = reinterpret_cast<const FBlendedCurveUtils&>(BlendedCurve1);
+
+		// Check arrays are not overlapping
+		check((void*)&InOutValueArray0 != (void*)&InValueArray1);
+
+		const int32 NumElements0 = InOutValueArray0.Num();	// ValueArray1 elements remain constant, but ValueArray0 can have entries added.
+		const int32 NumElements1 = InValueArray1.Num();
+
+		// A default element we re-use when an element from one of the two inputs is missing
+		FBlendedCurve::ElementType DefaultElement;
+
+		// Early out if we have no elements to union
+		if (NumElements1 == 0)
+		{
+			if (GVarUseBlendCurveFixes)
+			{
+				// NoTe: BEHAVIOR CHANGE START
+				//       FNamedValueArrayUtils::Union just returns, without applying InPredicate to the ElementPtr0(s),
+				//       resulting into the curve not blending to the default value
+				FBlendedCurve::ElementType* RESTRICT ElementPtr0 = InOutValueArray0.Elements.GetData();
+				const FBlendedCurve::ElementType* RESTRICT ElementEndPtr0 = ElementPtr0 + NumElements0;
+
+				while (ElementPtr0 != ElementEndPtr0)
+				{
+					DefaultElement.Name = ElementPtr0->Name;
+
+					InPredicate(*ElementPtr0, DefaultElement, UE::Anim::ENamedValueUnionFlags::ValidArg0);
+
+					++ElementPtr0;
+				}
+				// NoTe: BEHAVIOR CHANGE END
+			}
+		}
+		else
+		{
+			// Sort both input arrays if required
+			InOutValueArray0.SortElementsIfRequired();
+			InValueArray1.SortElementsIfRequired();
+
+			// Reserve memory for 1.5x combined curve counts.
+			// This can overestimate in some circumstances, but it handles the common cases which are:
+			// - One input is empty, the other not
+			// - Both inputs are non-empty but do not share most elements
+			int32 ReserveSize = FMath::Max(NumElements0, NumElements1);
+			ReserveSize += ReserveSize / 2;
+			InOutValueArray0.Reserve(ReserveSize);
+
+			// Use pointers to iterate as this uses fewer registers and this code is very hot
+			FBlendedCurve::ElementType* RESTRICT ElementPtr0 = InOutValueArray0.Elements.GetData();
+			const FBlendedCurve::ElementType* RESTRICT ElementPtr1 = InValueArray1.Elements.GetData();
+
+			const FBlendedCurve::ElementType* RESTRICT ElementEndPtr0 = ElementPtr0 + NumElements0;
+			const FBlendedCurve::ElementType* RESTRICT ElementEndPtr1 = ElementPtr1 + NumElements1;
+
+			// When we reach the end of either input arrays, we stop the tape merge and copy what remains
+			bool bIsDone = ElementPtr0 == ElementEndPtr0 || ElementPtr1 == ElementEndPtr1;
+
+			// Perform dual-iteration on the two sorted arrays
+			while (!bIsDone)
+			{
+				if (ElementPtr0->Name == ElementPtr1->Name)
+				{
+					// Elements match, run predicate and increment both indices
+					InPredicate(*ElementPtr0, *ElementPtr1, UE::Anim::ENamedValueUnionFlags::BothArgsValid);
+
+					++ElementPtr0;
+					++ElementPtr1;
+
+					bIsDone = ElementPtr0 == ElementEndPtr0 || ElementPtr1 == ElementEndPtr1;
+				}
+				else if (ElementPtr0->Name.FastLess(ElementPtr1->Name))
+				{
+					// ValueArray0 element is earlier, so run predicate with only ValueArray0 and increment ValueArray0
+					DefaultElement.Name = ElementPtr0->Name;
+
+					InPredicate(*ElementPtr0, DefaultElement, UE::Anim::ENamedValueUnionFlags::ValidArg0);
+
+					++ElementPtr0;
+
+					bIsDone = ElementPtr0 == ElementEndPtr0;
+				}
+				else
+				{
+					// ValueArray1 element is earlier, so add to ValueArray0, run predicate with only second and increment ValueArray1
+					const int32 ElementIndex0 = UE_PTRDIFF_TO_INT32(ElementPtr0 - InOutValueArray0.Elements.GetData());
+					InOutValueArray0.Elements.InsertUninitialized(ElementIndex0);
+
+					// Refresh pointers since they might have changed
+					ElementPtr0 = InOutValueArray0.Elements.GetData() + ElementIndex0;
+					ElementEndPtr0 = InOutValueArray0.Elements.GetData() + InOutValueArray0.Elements.Num();
+
+					// We use placement new to make sure the constructor is inlined to reduce redundant work
+					new(ElementPtr0) FBlendedCurve::ElementType();
+					ElementPtr0->Name = ElementPtr1->Name;
+
+					InPredicate(*ElementPtr0, *ElementPtr1, UE::Anim::ENamedValueUnionFlags::ValidArg1);
+
+					++ElementPtr0;	// Increment this as well since we've inserted
+					++ElementPtr1;
+
+					bIsDone = ElementPtr1 == ElementEndPtr1;
+				}
+			}
+
+			// Tape merge is done, copy anything that might be remaining
+			if (ElementPtr1 < ElementEndPtr1)
+			{
+				// Reached end of ValueArray0 with remaining in ValueArray1, we can just copy the remainder of ValueArray1
+				const int32 NumResults = InOutValueArray0.Elements.Num();
+				const int32 NumNewElements = (ElementEndPtr1 - ElementPtr1);
+				InOutValueArray0.Elements.Reserve(NumResults + NumNewElements);
+				InOutValueArray0.Elements.AddUninitialized(NumNewElements);
+
+				// Refresh pointers since they might have changed
+				ElementPtr0 = InOutValueArray0.Elements.GetData() + NumResults;
+				ElementEndPtr0 = ElementPtr0 + NumNewElements;
+
+				for (; ElementPtr1 < ElementEndPtr1; ++ElementPtr0, ++ElementPtr1)
+				{
+					// We use placement new to make sure the constructor is inlined to reduce redundant work
+					new(ElementPtr0) FBlendedCurve::ElementType();
+					ElementPtr0->Name = ElementPtr1->Name;
+
+					InPredicate(*ElementPtr0, *ElementPtr1, UE::Anim::ENamedValueUnionFlags::ValidArg1);
+				}
+			}
+
+			// NoTe: BEHAVIOR CHANGE START
+			if (GVarUseBlendCurveFixes && ElementPtr0 < ElementEndPtr0)
+			{
+				for (; ElementPtr0 < ElementEndPtr0; ++ElementPtr0)
+				{
+					DefaultElement.Name = ElementPtr0->Name;
+
+					InPredicate(*ElementPtr0, DefaultElement, UE::Anim::ENamedValueUnionFlags::ValidArg0);
+				}
+			}
+			// NoTe: BEHAVIOR CHANGE END
+		}
+
+		InOutValueArray0.CheckSorted();
+	}
+
+	// function adapted from TBaseBlendedCurve::LerpTo, that uses UnionEx instead of UE::Anim::FNamedValueArrayUtils::Union
+	static void LerpToEx(FBlendedCurve& InOutCurve, const FBlendedCurve& OtherCurve, float Alpha);
+
+	// LerpToEx with per linked bone weighting blend
+	static void LerpToPerBoneEx(FBlendedCurve& InOutCurve, const FBlendedCurve& OtherCurve, const FBoneContainer& BoneContainer, TConstArrayView<float> OtherCurveBoneWeights);
+};
+
+} // namespace UE::BlendStack
