blendstack - renaming CVarUseBlendCurveFixes to a.AnimNode.BlendStack.UseBlendCurveFixes

[RN] skip

diff --git a/Engine/Plugins/Animation/BlendStack/Source/Runtime/Private/BlendedCurveUtils.cpp b/Engine/Plugins/Animation/BlendStack/Source/Runtime/Private/BlendedCurveUtils.cpp
--- a/Engine/Plugins/Animation/BlendStack/Source/Runtime/Private/BlendedCurveUtils.cpp
+++ b/Engine/Plugins/Animation/BlendStack/Source/Runtime/Private/BlendedCurveUtils.cpp
@@ -9,7 +9,7 @@ namespace UE::BlendStack
 {
 	
 bool GVarUseBlendCurveFixes = true;
-static FAutoConsoleVariableRef CVarUseBlendCurveFixes(TEXT("a.AnimNode.BlendStack.UseBlendCurvePerBone"), GVarUseBlendCurveFixes, TEXT("Enable BlendStack BlendCurve fixes. It causes behavioral changes. For the good"));
+static FAutoConsoleVariableRef CVarUseBlendCurveFixes(TEXT("a.AnimNode.BlendStack.UseBlendCurveFixes"), GVarUseBlendCurveFixes, TEXT("Enable BlendStack BlendCurve fixes. It causes behavioral changes. For the good"));
 
 // function adapted from TBaseBlendedCurve::LerpTo, that uses UnionEx instead of UE::Anim::FNamedValueArrayUtils::Union
 void FBlendedCurveUtils::LerpToEx(FBlendedCurve& InOutCurve, const FBlendedCurve& OtherCurve, float Alpha)
