Re-enable forwarding in dead blending node (disabled by default)
Add a tag to inertialization requests to filter by node tag
Implement tag filtering in the blend stack inertialization request

diff --git a/Engine/Plugins/Animation/BlendStack/Source/Runtime/Private/AnimNode_BlendStack.cpp b/Engine/Plugins/Animation/BlendStack/Source/Runtime/Private/AnimNode_BlendStack.cpp
--- a/Engine/Plugins/Animation/BlendStack/Source/Runtime/Private/AnimNode_BlendStack.cpp
+++ b/Engine/Plugins/Animation/BlendStack/Source/Runtime/Private/AnimNode_BlendStack.cpp
@@ -943,7 +943,7 @@ float FAnimNode_BlendStack_Standalone::GetAccumulatedTime() const
 	return AnimPlayers.IsEmpty() ? 0.f : AnimPlayers[0].GetAccumulatedTime();
 }
 
-static void RequestInertialBlend(const FAnimationUpdateContext& Context, float BlendTime, const UBlendProfile* BlendProfile, EAlphaBlendOption BlendOption)
+static void RequestInertialBlend(const FAnimationUpdateContext& Context, float BlendTime, const UBlendProfile* BlendProfile, EAlphaBlendOption BlendOption, FName InertialBlendNodeTag)
 {
 #if ENABLE_ANIM_DEBUG
 	if (!UE::BlendStack::GVarAnimBlendStackEnable)
@@ -954,7 +954,26 @@ static void RequestInertialBlend(const FAnimationUpdateContext& Context, float B
 
 	if (BlendTime > 0.0f)
 	{
-		UE::Anim::IInertializationRequester* InertializationRequester = Context.GetMessage<UE::Anim::IInertializationRequester>();
+		UE::Anim::IInertializationRequester* InertializationRequester = nullptr;
+		// If was have an inertialization node tag, find the corresponding node
+		if (InertialBlendNodeTag != NAME_None)
+		{
+			InertializationRequester = Context.FindMessage<UE::Anim::IInertializationRequester>([&InertialBlendNodeTag](UE::Anim::IInertializationRequester& InMessage)
+				{
+					if (InMessage.GetTag() == InertialBlendNodeTag)
+					{
+						return true;
+					}
+					return false;
+				});
+			ensureMsgf(InertializationRequester, TEXT("FAnimNode_BlendStack with inertialization node tag '%s' could not find a matching node to inertialize"), *InertialBlendNodeTag.ToString());
+
+		}
+		else
+		{
+			InertializationRequester = Context.GetMessage<UE::Anim::IInertializationRequester>();
+		}
+		
 		if (InertializationRequester)
 		{
 			FInertializationRequest Request;
@@ -962,6 +981,7 @@ static void RequestInertialBlend(const FAnimationUpdateContext& Context, float B
 			Request.BlendProfile = BlendProfile;
 			Request.bUseBlendMode = true;
 			Request.BlendMode = BlendOption;
+			Request.Tag = InertialBlendNodeTag;
 #if ANIM_TRACE_ENABLED
 			Request.NodeId = Context.GetCurrentNodeId();
 			Request.AnimInstance = Context.AnimInstanceProxy->GetAnimInstanceObject();
@@ -977,7 +997,7 @@ static void RequestInertialBlend(const FAnimationUpdateContext& Context, float B
 }
 
 void FAnimNode_BlendStack_Standalone::BlendTo(const FAnimationUpdateContext& Context, UAnimationAsset* AnimationAsset, float AccumulatedTime, bool bLoop,
-	bool bMirrored, UMirrorDataTable* MirrorDataTable, float BlendTime, const UBlendProfile* BlendProfile, EAlphaBlendOption BlendOption, bool bUseInertialBlend,
+	bool bMirrored, UMirrorDataTable* MirrorDataTable, float BlendTime, const UBlendProfile* BlendProfile, EAlphaBlendOption BlendOption, bool bUseInertialBlend, FName InertialBlendNodeTag,
 	const FVector& BlendParameters, float PlayRate, float ActivationDelay, FName GroupName, EAnimGroupRole::Type GroupRole, EAnimSyncMethod GroupMethod, bool bOverridePositionWhenJoiningSyncGroupAsLeader)
 {
 	using namespace UE::Anim;
@@ -1009,13 +1029,13 @@ void FAnimNode_BlendStack_Standalone::BlendTo(const FAnimationUpdateContext& Con
 					{
 						// blend to the selected animation stitch
 						InternalBlendTo(Context, StitchAnimationAsset, SearchResult.TimeOffsetSeconds, false, SearchResult.bMirrored, MirrorDataTable,
-							StitchBlendTime, BlendProfile, BlendOption, bUseInertialBlend, BlendParameters, SearchResult.WantedPlayRate, ActivationDelay, GroupName, GroupRole, GroupMethod, bOverridePositionWhenJoiningSyncGroupAsLeader);
+							StitchBlendTime, BlendProfile, BlendOption, bUseInertialBlend, InertialBlendNodeTag, BlendParameters, SearchResult.WantedPlayRate, ActivationDelay, GroupName, GroupRole, GroupMethod, bOverridePositionWhenJoiningSyncGroupAsLeader);
 
 						// blend with an ActivationDelay of BlendTime - StitchBlendTime + ActivationDelay seconds
 						// to the AnimationAsset at AccumulatedTime + BlendTime - StitchBlendTime seconds in the future,
 						// so at BlendTime seconds ahead the AnimationAsset is playing the fully blended in pose at AccumulatedTime + BlendTime
 						InternalBlendTo(Context, AnimationAsset, AccumulatedTime + BlendTime - StitchBlendTime, bLoop, bMirrored, MirrorDataTable,
-							StitchBlendTime, BlendProfile, BlendOption, bUseInertialBlend, BlendParameters, PlayRate, BlendTime - StitchBlendTime + ActivationDelay, GroupName, GroupRole, GroupMethod, bOverridePositionWhenJoiningSyncGroupAsLeader);
+							StitchBlendTime, BlendProfile, BlendOption, bUseInertialBlend, InertialBlendNodeTag, BlendParameters, PlayRate, BlendTime - StitchBlendTime + ActivationDelay, GroupName, GroupRole, GroupMethod, bOverridePositionWhenJoiningSyncGroupAsLeader);
 
 						bNeedToBlendTo = false;
 					}
@@ -1043,12 +1063,12 @@ void FAnimNode_BlendStack_Standalone::BlendTo(const FAnimationUpdateContext& Con
 	if (bNeedToBlendTo)
 	{
 		InternalBlendTo(Context, AnimationAsset, AccumulatedTime, bLoop, bMirrored, MirrorDataTable,
-			BlendTime, BlendProfile, BlendOption, bUseInertialBlend, BlendParameters, PlayRate, ActivationDelay, GroupName, GroupRole, GroupMethod, bOverridePositionWhenJoiningSyncGroupAsLeader);
+			BlendTime, BlendProfile, BlendOption, bUseInertialBlend, InertialBlendNodeTag, BlendParameters, PlayRate, ActivationDelay, GroupName, GroupRole, GroupMethod, bOverridePositionWhenJoiningSyncGroupAsLeader);
 	}
 }
 
 void FAnimNode_BlendStack_Standalone::InternalBlendTo(const FAnimationUpdateContext& Context, UAnimationAsset* AnimationAsset, float AccumulatedTime, bool bLoop,
-	bool bMirrored, UMirrorDataTable* MirrorDataTable, float BlendTime, const UBlendProfile* BlendProfile, EAlphaBlendOption BlendOption, bool bUseInertialBlend,
+	bool bMirrored, UMirrorDataTable* MirrorDataTable, float BlendTime, const UBlendProfile* BlendProfile, EAlphaBlendOption BlendOption, bool bUseInertialBlend, FName InertialBlendNodeTag,
 	const FVector& BlendParameters, float PlayRate, float ActivationDelay, FName GroupName, EAnimGroupRole::Type GroupRole, EAnimSyncMethod GroupMethod, bool bOverridePositionWhenJoiningSyncGroupAsLeader)
 {
 	const bool bBlendStackIsEmpty = AnimPlayers.IsEmpty();
@@ -1061,7 +1081,7 @@ void FAnimNode_BlendStack_Standalone::InternalBlendTo(const FAnimationUpdateCont
 
 	if (bUseInertialBlend)
 	{
-		RequestInertialBlend(Context, BlendTime, BlendProfile, BlendOption);
+		RequestInertialBlend(Context, BlendTime, BlendProfile, BlendOption, InertialBlendNodeTag);
 		BlendTime = 0.0f;
 	}
 
@@ -1284,7 +1304,7 @@ bool FAnimNode_BlendStack::ConditionalBlendTo(const FAnimationUpdateContext& Con
 	if (bExecuteBlendTo)
 	{
 		BlendTo(Context, AnimationAsset, AnimationTime, bLoop, bMirrored, MirrorDataTable.Get(),
-			BlendTime, BlendProfile, BlendOption, bUseInertialBlend, BlendParameters, WantedPlayRate, ActivationDelayTime,
+			BlendTime, BlendProfile, BlendOption, bUseInertialBlend, InertialBlendNodeTag, BlendParameters, WantedPlayRate, ActivationDelayTime,
 			GetGroupName(), GetGroupRole(), GetGroupMethod());
 	}
 
diff --git a/Engine/Plugins/Animation/BlendStack/Source/Runtime/Private/BlendStackAnimNodeLibrary.cpp b/Engine/Plugins/Animation/BlendStack/Source/Runtime/Private/BlendStackAnimNodeLibrary.cpp
--- a/Engine/Plugins/Animation/BlendStack/Source/Runtime/Private/BlendStackAnimNodeLibrary.cpp
+++ b/Engine/Plugins/Animation/BlendStack/Source/Runtime/Private/BlendStackAnimNodeLibrary.cpp
@@ -86,6 +86,7 @@ void UBlendStackAnimNodeLibrary::BlendTo(const FAnimUpdateContext& Context,
 					BlendStackNodePtr->BlendProfile,
 					BlendStackNodePtr->BlendOption,
 					BlendStackNodePtr->bUseInertialBlend,
+					BlendStackNodePtr->InertialBlendNodeTag,
 					BlendParameters,
 					WantedPlayRate,
 					ActivationDelay,
@@ -134,6 +135,7 @@ void UBlendStackAnimNodeLibrary::BlendToWithSettings(const FAnimUpdateContext& C
 					BlendProfile,
 					BlendOption,
 					bInertialBlend,
+					BlendStackNodePtr->InertialBlendNodeTag,
 					BlendParameters,
 					WantedPlayRate,
 					ActivationDelay,
diff --git a/Engine/Plugins/Animation/BlendStack/Source/Runtime/Public/BlendStack/AnimNode_BlendStack.h b/Engine/Plugins/Animation/BlendStack/Source/Runtime/Public/BlendStack/AnimNode_BlendStack.h
--- a/Engine/Plugins/Animation/BlendStack/Source/Runtime/Public/BlendStack/AnimNode_BlendStack.h
+++ b/Engine/Plugins/Animation/BlendStack/Source/Runtime/Public/BlendStack/AnimNode_BlendStack.h
@@ -169,7 +169,7 @@ struct BLENDSTACK_API FAnimNode_BlendStack_Standalone : public FAnimNode_AssetPl
 	void BlendTo(const FAnimationUpdateContext& Context, UAnimationAsset* AnimationAsset, float AccumulatedTime = 0.f, bool bLoop = false, 
 		bool bMirrored = false, UMirrorDataTable* MirrorDataTable = nullptr, float BlendTime = 0.2f,
 		const UBlendProfile* BlendProfile = nullptr, EAlphaBlendOption BlendOption = EAlphaBlendOption::Linear, 
-		bool bUseInertialBlend = false, const FVector& BlendParameters = FVector::Zero(), float PlayRate = 1.f, float ActivationDelay = 0.f,
+		bool bUseInertialBlend = false, FName InertialBlendNodeTag = NAME_None, const FVector& BlendParameters = FVector::Zero(), float PlayRate = 1.f, float ActivationDelay = 0.f,
 		FName GroupName = NAME_None, EAnimGroupRole::Type GroupRole = EAnimGroupRole::CanBeLeader, EAnimSyncMethod Method = EAnimSyncMethod::DoNotSync, bool bOverridePositionWhenJoiningSyncGroupAsLeader = false);
 	void UpdatePlayRate(float PlayRate);
 	virtual void Reset();
@@ -191,7 +191,7 @@ protected:
 	void InternalBlendTo(const FAnimationUpdateContext& Context, UAnimationAsset* AnimationAsset, float AccumulatedTime, bool bLoop, 
 		bool bMirrored, UMirrorDataTable* MirrorDataTable, float BlendTime,
 		const UBlendProfile* BlendProfile, EAlphaBlendOption BlendOption, 
-		bool bUseInertialBlend, const FVector& BlendParameters, float PlayRate, float ActivationDelay,
+		bool bUseInertialBlend, FName InertialBlendNodeTag, const FVector& BlendParameters, float PlayRate, float ActivationDelay,
 		FName GroupName, EAnimGroupRole::Type GroupRole, EAnimSyncMethod Method, bool bOverridePositionWhenJoiningSyncGroupAsLeader);
 
 	// Call this to update current blendspace player xy blend parameters. By default, we only update them on initial BlendTo.
@@ -297,10 +297,14 @@ struct BLENDSTACK_API FAnimNode_BlendStack : public FAnimNode_BlendStack_Standal
 	UPROPERTY(EditAnywhere, Category = Blendspace)
 	float BlendParametersDeltaThreshold = 0.0f;
 
-	// tunable animation transition blend time 
+	// Enable to use inertial blending
 	UPROPERTY(EditAnywhere, Category = Settings)
 	bool bUseInertialBlend = false;
 
+	// Tag to force a specific inertialization / dead blending node to process inertial blend requests coming from this blend stack
+	UPROPERTY(EditAnywhere, Category = Settings)
+	FName InertialBlendNodeTag = NAME_None;
+
 	// Reset the blend stack if it has become relevant to the graph after not being updated on previous frames.
 	UPROPERTY(EditAnywhere, Category = Settings)
 	bool bResetOnBecomingRelevant = true;
diff --git a/Engine/Plugins/Animation/PoseSearch/Source/Runtime/Private/AnimNode_MotionMatching.cpp b/Engine/Plugins/Animation/PoseSearch/Source/Runtime/Private/AnimNode_MotionMatching.cpp
--- a/Engine/Plugins/Animation/PoseSearch/Source/Runtime/Private/AnimNode_MotionMatching.cpp
+++ b/Engine/Plugins/Animation/PoseSearch/Source/Runtime/Private/AnimNode_MotionMatching.cpp
@@ -266,7 +266,7 @@ void FAnimNode_MotionMatching::UpdateAssetPlayer(const FAnimationUpdateContext&
 				{
 					FAnimNode_BlendStack_Standalone::BlendTo(Context, AnimationAsset, MotionMatchingState.CurrentSearchResult.AssetTime,
 						SearchIndexAsset->IsLooping(), SearchIndexAsset->IsMirrored(), CurrentResultDatabase->Schema->GetMirrorDataTable(DefaultRole), BlendTime,
-						BlendProfile, BlendOption, bUseInertialBlend, SearchIndexAsset->GetBlendParameters(), MotionMatchingState.WantedPlayRate * PlayRateMultiplier, 0,
+						BlendProfile, BlendOption, bUseInertialBlend, NAME_None, SearchIndexAsset->GetBlendParameters(), MotionMatchingState.WantedPlayRate * PlayRateMultiplier, 0,
 						GetGroupName(), GetGroupRole(), GetGroupMethod(), GetOverridePositionWhenJoiningSyncGroupAsLeader());
 				}
 			}
diff --git a/Engine/Plugins/Animation/PoseSearch/Source/Runtime/Private/AnimNode_MotionMatchingInteraction.cpp b/Engine/Plugins/Animation/PoseSearch/Source/Runtime/Private/AnimNode_MotionMatchingInteraction.cpp
--- a/Engine/Plugins/Animation/PoseSearch/Source/Runtime/Private/AnimNode_MotionMatchingInteraction.cpp
+++ b/Engine/Plugins/Animation/PoseSearch/Source/Runtime/Private/AnimNode_MotionMatchingInteraction.cpp
@@ -127,7 +127,7 @@ void FAnimNode_MotionMatchingInteraction::UpdateAssetPlayer(const FAnimationUpda
 			check(RoledSkeleton);
 
 			BlendTo(Context, RoledAnimAsset, CurrentResult.SelectedTime, CurrentResult.bLoop, CurrentResult.bIsMirrored, RoledSkeleton->MirrorDataTable.Get(),
-				BlendTime, BlendProfile, BlendOption, bUseInertialBlend, CurrentResult.BlendParameters, CurrentResult.WantedPlayRate);
+				BlendTime, BlendProfile, BlendOption, bUseInertialBlend, NAME_None, CurrentResult.BlendParameters, CurrentResult.WantedPlayRate);
 
 			bBlendToExecuted = true;
 		}
diff --git a/Engine/Plugins/Chooser/Source/Chooser/Private/AnimNode_ChooserPlayer.cpp b/Engine/Plugins/Chooser/Source/Chooser/Private/AnimNode_ChooserPlayer.cpp
--- a/Engine/Plugins/Chooser/Source/Chooser/Private/AnimNode_ChooserPlayer.cpp
+++ b/Engine/Plugins/Chooser/Source/Chooser/Private/AnimNode_ChooserPlayer.cpp
@@ -166,7 +166,7 @@ void FAnimNode_ChooserPlayer::UpdateAssetPlayer(const FAnimationUpdateContext& C
 		
 		FAnimNode_BlendStack_Standalone::BlendTo(Context, NewAsset, Settings.StartTime,
 			bLoop, Settings.bMirror, MirrorDataTable, Settings.BlendTime,
-			Settings.BlendProfile, Settings.BlendOption, Settings.bUseInertialBlend, FVector::Zero(), Settings.PlaybackRate, 0.f,
+			Settings.BlendProfile, Settings.BlendOption, Settings.bUseInertialBlend, NAME_None, FVector::Zero(), Settings.PlaybackRate, 0.f,
 			GetGroupName(), GetGroupRole(), GetGroupMethod());
 
 		if (!Settings.CurveOverrides.Values.IsEmpty())
diff --git a/Engine/Source/Editor/AnimGraph/Private/AnimGraphNode_DeadBlending.cpp b/Engine/Source/Editor/AnimGraph/Private/AnimGraphNode_DeadBlending.cpp
--- a/Engine/Source/Editor/AnimGraph/Private/AnimGraphNode_DeadBlending.cpp
+++ b/Engine/Source/Editor/AnimGraph/Private/AnimGraphNode_DeadBlending.cpp
@@ -1,6 +1,7 @@
 // Copyright Epic Games, Inc. All Rights Reserved.
 
 #include "AnimGraphNode_DeadBlending.h"
+#include "IAnimBlueprintCopyTermDefaultsContext.h"
 
 #define LOCTEXT_NAMESPACE "AnimGraphNode_DeadBlending"
 
@@ -30,4 +31,11 @@ void UAnimGraphNode_DeadBlending::GetInputLinkAttributes(FNodeAttributeArray& Ou
 	OutAttributes.Add(UE::Anim::IInertializationRequester::Attribute);
 }
 
+void UAnimGraphNode_DeadBlending::OnCopyTermDefaultsToDefaultObject(IAnimBlueprintCopyTermDefaultsContext& InCompilationContext, IAnimBlueprintNodeCopyTermDefaultsContext& InPerNodeContext, IAnimBlueprintGeneratedClassCompiledData& OutCompiledData)
+{
+	UAnimGraphNode_DeadBlending* TrueNode = InCompilationContext.GetMessageLog().FindSourceObjectTypeChecked<UAnimGraphNode_DeadBlending>(this);
+
+	FAnimNode_DeadBlending* DestinationNode = reinterpret_cast<FAnimNode_DeadBlending*>(InPerNodeContext.GetDestinationPtr());
+	DestinationNode->SetTag(TrueNode->GetTag());
+}
 #undef LOCTEXT_NAMESPACE
diff --git a/Engine/Source/Editor/AnimGraph/Private/AnimGraphNode_Inertialization.cpp b/Engine/Source/Editor/AnimGraph/Private/AnimGraphNode_Inertialization.cpp
--- a/Engine/Source/Editor/AnimGraph/Private/AnimGraphNode_Inertialization.cpp
+++ b/Engine/Source/Editor/AnimGraph/Private/AnimGraphNode_Inertialization.cpp
@@ -1,6 +1,7 @@
 // Copyright Epic Games, Inc. All Rights Reserved.
 
 #include "AnimGraphNode_Inertialization.h"
+#include "IAnimBlueprintCopyTermDefaultsContext.h"
 
 #define LOCTEXT_NAMESPACE "AnimGraphNode_Inertialization"
 
@@ -30,4 +31,11 @@ void UAnimGraphNode_Inertialization::GetInputLinkAttributes(FNodeAttributeArray&
 	OutAttributes.Add(UE::Anim::IInertializationRequester::Attribute);
 }
 
+void UAnimGraphNode_Inertialization::OnCopyTermDefaultsToDefaultObject(IAnimBlueprintCopyTermDefaultsContext& InCompilationContext, IAnimBlueprintNodeCopyTermDefaultsContext& InPerNodeContext, IAnimBlueprintGeneratedClassCompiledData& OutCompiledData)
+{
+	UAnimGraphNode_Inertialization* TrueNode = InCompilationContext.GetMessageLog().FindSourceObjectTypeChecked<UAnimGraphNode_Inertialization>(this);
+
+	FAnimNode_Inertialization* DestinationNode = reinterpret_cast<FAnimNode_Inertialization*>(InPerNodeContext.GetDestinationPtr());
+	DestinationNode->SetTag(TrueNode->GetTag());
+}
 #undef LOCTEXT_NAMESPACE
diff --git a/Engine/Source/Editor/AnimGraph/Public/AnimGraphNode_DeadBlending.h b/Engine/Source/Editor/AnimGraph/Public/AnimGraphNode_DeadBlending.h
--- a/Engine/Source/Editor/AnimGraph/Public/AnimGraphNode_DeadBlending.h
+++ b/Engine/Source/Editor/AnimGraph/Public/AnimGraphNode_DeadBlending.h
@@ -8,7 +8,6 @@
 #include "Animation/AnimNode_DeadBlending.h"
 #include "AnimGraphNode_DeadBlending.generated.h"
 
-
 UCLASS(MinimalAPI)
 class UAnimGraphNode_DeadBlending : public UAnimGraphNode_Base
 {
@@ -22,4 +21,5 @@ class UAnimGraphNode_DeadBlending : public UAnimGraphNode_Base
 	virtual FText GetNodeTitle(ENodeTitleType::Type TitleType) const override;
 	virtual FText GetMenuCategory() const override;
 	virtual void GetInputLinkAttributes(FNodeAttributeArray& OutAttributes) const override;
+	virtual void OnCopyTermDefaultsToDefaultObject(IAnimBlueprintCopyTermDefaultsContext& InCompilationContext, IAnimBlueprintNodeCopyTermDefaultsContext& InPerNodeContext, IAnimBlueprintGeneratedClassCompiledData& OutCompiledData) override;
 };
diff --git a/Engine/Source/Editor/AnimGraph/Public/AnimGraphNode_Inertialization.h b/Engine/Source/Editor/AnimGraph/Public/AnimGraphNode_Inertialization.h
--- a/Engine/Source/Editor/AnimGraph/Public/AnimGraphNode_Inertialization.h
+++ b/Engine/Source/Editor/AnimGraph/Public/AnimGraphNode_Inertialization.h
@@ -8,7 +8,6 @@
 #include "Animation/AnimNode_Inertialization.h"
 #include "AnimGraphNode_Inertialization.generated.h"
 
-
 UCLASS(MinimalAPI)
 class UAnimGraphNode_Inertialization : public UAnimGraphNode_Base
 {
@@ -22,4 +21,6 @@ class UAnimGraphNode_Inertialization : public UAnimGraphNode_Base
 	virtual FText GetNodeTitle(ENodeTitleType::Type TitleType) const override;
 	virtual FText GetMenuCategory() const override;
 	virtual void GetInputLinkAttributes(FNodeAttributeArray& OutAttributes) const override;
+	virtual void OnCopyTermDefaultsToDefaultObject(IAnimBlueprintCopyTermDefaultsContext& InCompilationContext, IAnimBlueprintNodeCopyTermDefaultsContext& InPerNodeContext, IAnimBlueprintGeneratedClassCompiledData& OutCompiledData) override;
+
 };
diff --git a/Engine/Source/Runtime/Engine/Classes/Animation/AnimInertializationRequest.h b/Engine/Source/Runtime/Engine/Classes/Animation/AnimInertializationRequest.h
--- a/Engine/Source/Runtime/Engine/Classes/Animation/AnimInertializationRequest.h
+++ b/Engine/Source/Runtime/Engine/Classes/Animation/AnimInertializationRequest.h
@@ -67,6 +67,10 @@ struct FInertializationRequest
 	UPROPERTY(Transient)
 	TObjectPtr<UCurveFloat> CustomBlendCurve = nullptr;
 
+	// Inertialization / Dead blend node tag to force a particular node to handle the request.
+	UPROPERTY(Transient)
+	FName Tag = NAME_None;
+
 // if ANIM_TRACE_ENABLED - these properties are only used for debugging when ANIM_TRACE_ENABLED == 1
 	
 	UE_DEPRECATED(5.4, "Use DescriptionString instead.")
diff --git a/Engine/Source/Runtime/Engine/Classes/Animation/AnimNodeBase.h b/Engine/Source/Runtime/Engine/Classes/Animation/AnimNodeBase.h
--- a/Engine/Source/Runtime/Engine/Classes/Animation/AnimNodeBase.h
+++ b/Engine/Source/Runtime/Engine/Classes/Animation/AnimNodeBase.h
@@ -238,7 +238,30 @@ public:
 		
 		return nullptr;
 	}
-	
+
+	// Find the innermost scoped message of the specified type matching the condition of InFunction
+	template<typename TGraphMessageType>
+	TGraphMessageType* FindMessage(TFunctionRef<bool(TGraphMessageType&)> InFunction) const
+	{
+		if (ensure(SharedContext != nullptr))
+		{
+			TGraphMessageType* Message = nullptr;
+			SharedContext->MessageStack.ForEachMessage<TGraphMessageType>([&Message, &InFunction](TGraphMessageType& InMessage)
+			{
+				if (InFunction(InMessage))
+				{
+					Message = &InMessage;
+					return UE::Anim::FMessageStack::EEnumerate::Stop;
+				}
+				return UE::Anim::FMessageStack::EEnumerate::Continue;
+			});
+				
+			return Message;
+		}
+
+		return nullptr;
+	}
+
 	// Get the innermost scoped message of the specified type
 	template<typename TGraphMessageType>
 	TGraphMessageType& GetMessageChecked() const
diff --git a/Engine/Source/Runtime/Engine/Classes/Animation/AnimNode_DeadBlending.h b/Engine/Source/Runtime/Engine/Classes/Animation/AnimNode_DeadBlending.h
--- a/Engine/Source/Runtime/Engine/Classes/Animation/AnimNode_DeadBlending.h
+++ b/Engine/Source/Runtime/Engine/Classes/Animation/AnimNode_DeadBlending.h
@@ -36,6 +36,11 @@ struct FAnimNode_DeadBlending : public FAnimNode_Base, public IBoneReferenceSkel
 	UPROPERTY(EditAnywhere, Category = Links)
 	FPoseLink Source;
 
+#if WITH_EDITORONLY_DATA
+	void SetTag(FName InTag) { Tag = InTag; }
+#endif
+	FName GetTag() { return Tag; }
+
 private:
 
 	// When enabled, the default blend settings will always be used rather than those coming from the inertialization request.
@@ -143,6 +148,19 @@ private:
 	UPROPERTY(EditAnywhere, Category = Blending)
 	bool bResetOnBecomingRelevant = true;
 
+	/**
+	* When enabled this option will forward inertialization requests through any downstream UseCachedPose nodes which 
+	* have had their update skipped (e.g. because they have already been updated in another location). This can be
+	* useful in the case where the same cached pose is used in multiple places, and having an inertialization request 
+	* that goes with it caught in only one of those places would create popping.
+	*/
+	UPROPERTY(EditAnywhere, Category = Requests)
+	bool bForwardRequestsThroughSkippedCachedPoseNodes = false;
+
+	/* UAnimGraphNode_Deadblending tag to match at runtime for filtering inertialization requests */
+	UPROPERTY()
+	FName Tag = NAME_None;
+
 #if WITH_EDITORONLY_DATA
 	
 	// This setting can be used to show what the extrapolation of the animation looks like.
diff --git a/Engine/Source/Runtime/Engine/Classes/Animation/AnimNode_Inertialization.h b/Engine/Source/Runtime/Engine/Classes/Animation/AnimNode_Inertialization.h
--- a/Engine/Source/Runtime/Engine/Classes/Animation/AnimNode_Inertialization.h
+++ b/Engine/Source/Runtime/Engine/Classes/Animation/AnimNode_Inertialization.h
@@ -42,6 +42,8 @@ public:
 
 	// Add a record of this request
 	virtual void AddDebugRecord(const FAnimInstanceProxy& InSourceProxy, int32 InSourceNodeId) = 0;
+
+	virtual FName GetTag() const = 0;
 };
 
 }	// namespace UE::Anim
@@ -302,6 +304,11 @@ struct FAnimNode_Inertialization : public FAnimNode_Base, public IBoneReferenceS
 	UPROPERTY(EditAnywhere, Category = Links)
 	FPoseLink Source;
 
+#if WITH_EDITORONLY_DATA
+	void SetTag(FName InTag) { Tag = InTag; }
+#endif
+	FName GetTag() { return Tag; }
+
 private:
 
 	// Optional default blend profile to use when no blend profile is supplied with the inertialization request
@@ -337,6 +344,9 @@ private:
 	UPROPERTY(EditAnywhere, Category = Requests)
 	bool bForwardRequestsThroughSkippedCachedPoseNodes = true;
 
+	UPROPERTY()
+	FName Tag = NAME_None;
+
 public: // FAnimNode_Inertialization
 
 	// Note: We need to explicitly disable warnings on these constructors/operators for clang to be happy with deprecated variables
diff --git a/Engine/Source/Runtime/Engine/Private/Animation/AnimNode_DeadBlending.cpp b/Engine/Source/Runtime/Engine/Private/Animation/AnimNode_DeadBlending.cpp
--- a/Engine/Source/Runtime/Engine/Private/Animation/AnimNode_DeadBlending.cpp
+++ b/Engine/Source/Runtime/Engine/Private/Animation/AnimNode_DeadBlending.cpp
@@ -39,12 +39,12 @@ namespace UE::Anim
 			Node.RequestInertialization(FInertializationRequest(InRequestedDuration, InBlendProfile));
 		}
 
-		virtual void RequestInertialization(const FInertializationRequest& Request)
+		virtual void RequestInertialization(const FInertializationRequest& Request) override
 		{
 			Node.RequestInertialization(Request);
 		}
 
-		virtual void AddDebugRecord(const FAnimInstanceProxy& InSourceProxy, int32 InSourceNodeId)
+		virtual void AddDebugRecord(const FAnimInstanceProxy& InSourceProxy, int32 InSourceNodeId) override
 		{
 #if WITH_EDITORONLY_DATA
 			Proxy.RecordNodeAttribute(InSourceProxy, NodeId, InSourceNodeId, IInertializationRequester::Attribute);
@@ -52,6 +52,8 @@ namespace UE::Anim
 			TRACE_ANIM_NODE_ATTRIBUTE(Proxy, InSourceProxy, NodeId, InSourceNodeId, IInertializationRequester::Attribute);
 		}
 
+		virtual FName GetTag() const override { return Node.GetTag(); }
+
 		// Node to target
 		FAnimNode_DeadBlending& Node;
 
@@ -706,7 +708,11 @@ void FAnimNode_DeadBlending::RequestInertialization(const FInertializationReques
 {
 	if (Request.Duration >= 0.0f)
 	{
-		RequestQueue.AddUnique(Request);
+		// If the request has a specific tag, it needs to match ours
+		if (Request.Tag == NAME_None || Request.Tag == GetTag())
+		{
+			RequestQueue.AddUnique(Request);
+		}
 	}
 }
 
@@ -809,7 +815,46 @@ void FAnimNode_DeadBlending::Update_AnyThread(const FAnimationUpdateContext& Con
 	// Catch the inertialization request message and call the node's RequestInertialization function with the request
 	UE::Anim::TScopedGraphMessage<UE::Anim::FDeadBlendingRequester> InertializationMessage(Context, Context, this);
 
-	Source.Update(Context);
+	if (bForwardRequestsThroughSkippedCachedPoseNodes)
+	{
+		const int32 NodeId = Context.GetCurrentNodeId();
+		const FAnimInstanceProxy& Proxy = *Context.AnimInstanceProxy;
+
+		// Handle skipped updates for cached poses by forwarding to inertialization nodes in those residual stacks
+		UE::Anim::TScopedGraphMessage<UE::Anim::FCachedPoseSkippedUpdateHandler> CachedPoseSkippedUpdate(Context, [this, NodeId, &Proxy](TArrayView<const UE::Anim::FMessageStack> InSkippedUpdates)
+		{
+			// If we have a pending request forward the request to other Inertialization nodes
+			// that were skipped due to pose caching.
+			if (RequestQueue.Num() > 0)
+			{
+				// Cached poses have their Update function called once even though there may be multiple UseCachedPose nodes for the same pose.
+				// Because of this, there may be Inertialization ancestors of the UseCachedPose nodes that missed out on requests.
+				// So here we forward 'this' node's requests to the ancestors of those skipped UseCachedPose nodes.
+				// Note that in some cases, we may be forwarding the requests back to this same node.  Those duplicate requests will ultimately
+				// be ignored by the 'AddUnique' in the body of FAnimNode_DeadBlending::RequestInertialization.
+				for (const UE::Anim::FMessageStack& Stack : InSkippedUpdates)
+				{
+					Stack.ForEachMessage<UE::Anim::IInertializationRequester>([this, NodeId, &Proxy](UE::Anim::IInertializationRequester& InMessage)
+					{
+						for (const FInertializationRequest& Request : RequestQueue)
+						{
+							InMessage.RequestInertialization(Request);
+						}
+						InMessage.AddDebugRecord(Proxy, NodeId);
+
+						return UE::Anim::FMessageStack::EEnumerate::Stop;
+					});
+				}
+			}
+		});
+
+		// Context message stack lifetime is scope based so we need to call Source.Update() before exiting the scope of the message above.
+		Source.Update(Context);
+	}
+	else
+	{
+		Source.Update(Context);
+	}
 
 	// Accumulate delta time between calls to Evaluate_AnyThread
 	DeltaTime += Context.GetDeltaTime();
diff --git a/Engine/Source/Runtime/Engine/Private/Animation/AnimNode_Inertialization.cpp b/Engine/Source/Runtime/Engine/Private/Animation/AnimNode_Inertialization.cpp
--- a/Engine/Source/Runtime/Engine/Private/Animation/AnimNode_Inertialization.cpp
+++ b/Engine/Source/Runtime/Engine/Private/Animation/AnimNode_Inertialization.cpp
@@ -58,7 +58,7 @@ namespace UE::Anim
 			Node.RequestInertialization(InInertializationRequest);
 		}
 
-		virtual void AddDebugRecord(const FAnimInstanceProxy& InSourceProxy, int32 InSourceNodeId)
+		virtual void AddDebugRecord(const FAnimInstanceProxy& InSourceProxy, int32 InSourceNodeId) override
 		{
 #if WITH_EDITORONLY_DATA
 			Proxy.RecordNodeAttribute(InSourceProxy, NodeId, InSourceNodeId, IInertializationRequester::Attribute);
@@ -66,6 +66,8 @@ namespace UE::Anim
 			TRACE_ANIM_NODE_ATTRIBUTE(Proxy, InSourceProxy, NodeId, InSourceNodeId, IInertializationRequester::Attribute);
 		}
 
+		virtual FName GetTag() const override { return Node.GetTag(); }
+
 		// Node to target
 		FAnimNode_Inertialization& Node;
 
@@ -272,7 +274,11 @@ void FAnimNode_Inertialization::RequestInertialization(const FInertializationReq
 {
 	if (Request.Duration >= 0.0f)
 	{
-		RequestQueue.AddUnique(Request);
+		// If the request has a specific tag, it needs to match ours
+		if (Request.Tag == NAME_None || Request.Tag == GetTag())
+		{
+			RequestQueue.AddUnique(Request);
+		}
 	}
 }
 
