BP Expose ability to tell if current anim player in blendstack is active via 'GetCurrentBlendStackAnimIsActive'

#tests Editor

diff --git a/Engine/Plugins/Animation/BlendStack/Source/Runtime/Private/BlendStackAnimNodeLibrary.cpp b/Engine/Plugins/Animation/BlendStack/Source/Runtime/Private/BlendStackAnimNodeLibrary.cpp
--- a/Engine/Plugins/Animation/BlendStack/Source/Runtime/Private/BlendStackAnimNodeLibrary.cpp
+++ b/Engine/Plugins/Animation/BlendStack/Source/Runtime/Private/BlendStackAnimNodeLibrary.cpp
@@ -34,6 +34,18 @@ float UBlendStackAnimNodeLibrary::GetCurrentBlendStackAnimAssetTime(const FAnimN
 	return 0.f;
 }
 
+bool UBlendStackAnimNodeLibrary::GetCurrentBlendStackAnimIsActive(const FAnimNodeReference& Node)
+{
+	if (FAnimNode_BlendStackInput* BlendStackInput = Node.GetAnimNodePtr<FAnimNode_BlendStackInput>())
+	{
+		if (BlendStackInput->Player && *BlendStackInput->Player)
+		{
+			return (*BlendStackInput->Player)->IsActive();
+		}
+	}
+	return false;
+}
+
 bool UBlendStackAnimNodeLibrary::GetCurrentBlendStackAnimAssetMirrored(const FAnimNodeReference& Node)
 {
 	if (FAnimNode_BlendStackInput* BlendStackInput = Node.GetAnimNodePtr<FAnimNode_BlendStackInput>())
diff --git a/Engine/Plugins/Animation/BlendStack/Source/Runtime/Public/BlendStack/BlendStackAnimNodeLibrary.h b/Engine/Plugins/Animation/BlendStack/Source/Runtime/Public/BlendStack/BlendStackAnimNodeLibrary.h
--- a/Engine/Plugins/Animation/BlendStack/Source/Runtime/Public/BlendStack/BlendStackAnimNodeLibrary.h
+++ b/Engine/Plugins/Animation/BlendStack/Source/Runtime/Public/BlendStack/BlendStackAnimNodeLibrary.h
@@ -38,6 +38,10 @@ public:
     UFUNCTION(BlueprintPure, Category = "Animation|BlendStack", meta = (BlueprintThreadSafe))
     static float GetCurrentBlendStackAnimAssetTime(const FAnimNodeReference& Node);
 
+	/** Get if current anim is active */
+	UFUNCTION(BlueprintPure, Category = "Animation|BlendStack", meta = (BlueprintThreadSafe))
+	static bool GetCurrentBlendStackAnimIsActive(const FAnimNodeReference& Node);
+
 	/** Get if we are currently mirrored from a Blend Stack Input node */
 	UFUNCTION(BlueprintPure, Category = "Animation|BlendStack", meta = (BlueprintThreadSafe))
 	static bool GetCurrentBlendStackAnimAssetMirrored(const FAnimNodeReference& Node);
