fix root motion sampling in blend stack

diff --git a/Engine/Plugins/Animation/BlendStack/Source/Runtime/Private/AnimNode_BlendStack.cpp b/Engine/Plugins/Animation/BlendStack/Source/Runtime/Private/AnimNode_BlendStack.cpp
--- a/Engine/Plugins/Animation/BlendStack/Source/Runtime/Private/AnimNode_BlendStack.cpp
+++ b/Engine/Plugins/Animation/BlendStack/Source/Runtime/Private/AnimNode_BlendStack.cpp
@@ -793,7 +793,7 @@ void FAnimNode_BlendStack_Standalone::UpdateAssetPlayer(const FAnimationUpdateCo
 		const float BlendInPercentage = bIsLastAnimPlayers ? 1.f : AnimPlayer.GetBlendInWeight();
 		const float AnimPlayerBlendWeight = CurrentWeightMultiplier * BlendInPercentage;
 
-		FAnimationUpdateContext AnimPlayerContext = Context.FractionalWeightAndRootMotion(AnimPlayerBlendWeight, AnimPlayerBlendWeight);
+		FAnimationUpdateContext AnimPlayerContext = Context.FractionalWeight(AnimPlayerBlendWeight);
 		{
 
 #if ENABLE_ANIM_DEBUG || ENABLE_VISUAL_LOG 
