PoseSearch - using array views instead of arrays to facilitate the transition to allocator agnostic containers

diff --git a/Engine/Plugins/Animation/BlendStack/Source/Runtime/Private/BlendStackAnimEventsFilterScope.h b/Engine/Plugins/Animation/BlendStack/Source/Runtime/Private/BlendStackAnimEventsFilterScope.h
--- a/Engine/Plugins/Animation/BlendStack/Source/Runtime/Private/BlendStackAnimEventsFilterScope.h
+++ b/Engine/Plugins/Animation/BlendStack/Source/Runtime/Private/BlendStackAnimEventsFilterScope.h
@@ -12,13 +12,15 @@ namespace UE { namespace Anim {
 	class FBlendStackAnimEventsFilterContext : public IAnimEventsFilterContext
 	{
 	public:
-
 		FBlendStackAnimEventsFilterContext() {}
 		FBlendStackAnimEventsFilterContext(const TSharedRef<TArray<FName>> & InFiredNotifies, const TSharedRef<TMap<FName,float>>& InNotifyBanList);
 		
 		virtual bool ShouldFilterNotify(const FAnimNotifyEventReference& InNotifyEventRef) const override;
 	
+	private:
+		// @todo: add inline allocator with some preallocated memory
 		TSharedRef<TArray<FName>> FiredNotifies;
+		// @todo: add inline allocator with some preallocated memory
 		TSharedRef<TMap<FName,float>> NotifyBanList;
 	};
 	
@@ -27,14 +29,15 @@ namespace UE { namespace Anim {
 		DECLARE_ANIMGRAPH_MESSAGE(FBlendStackAnimEventsFilterScope);
 
 	public:
-		
 		FBlendStackAnimEventsFilterScope(const TSharedRef<TArray<FName>> & InFiredNotifies, const TSharedRef<TMap<FName,float>> & InNotifyBanList);
 		FBlendStackAnimEventsFilterScope(const TSharedPtr<TArray<FName>> & InFiredNotifies, const TSharedPtr<TMap<FName,float>> & InNotifyBanList);
 		
 		virtual TUniquePtr<const IAnimNotifyEventContextDataInterface> MakeUniqueEventContextData() const override;
 
 	private:
+		// @todo: add inline allocator with some preallocated memory
 		TSharedRef<TArray<FName>> FiredNotifies;
+		// @todo: add inline allocator with some preallocated memory
 		TSharedRef<TMap<FName,float>> NotifyBanList;
 	};
 	
diff --git a/Engine/Plugins/Animation/PoseSearch/Source/Editor/Private/PoseSearchDatabaseAssetListItem.cpp b/Engine/Plugins/Animation/PoseSearch/Source/Editor/Private/PoseSearchDatabaseAssetListItem.cpp
--- a/Engine/Plugins/Animation/PoseSearch/Source/Editor/Private/PoseSearchDatabaseAssetListItem.cpp
+++ b/Engine/Plugins/Animation/PoseSearch/Source/Editor/Private/PoseSearchDatabaseAssetListItem.cpp
@@ -58,7 +58,7 @@ namespace UE::PoseSearch
 			}
 			else
 			{
-				const TArray<FBlendSample>& BlendSamples = BlendSpace->GetBlendSamples();
+				const TConstArrayView<FBlendSample> BlendSamples = BlendSpace->GetBlendSamples();
 				for (int i = 0; i < BlendSamples.Num() - 1; ++i)
 				{
 					const FBlendSample& CurrSample = BlendSamples[i];
@@ -862,8 +862,6 @@ namespace UE::PoseSearch
 				ViewModelPtr->SetDisableReselection(AssetTreeNode->SourceAssetIdx, NewCheckboxState == ECheckBoxState::Checked ? true : false);
 
 				SkeletonView.Pin()->RefreshTreeView(false, true);
-
-				// no need to rebuild the SearchIndex (ViewModelPtr->BuildSearchIndex()), since bDisableReselection is a runtime only parameter
 			}
 		}
 	}
diff --git a/Engine/Plugins/Animation/PoseSearch/Source/Editor/Private/PoseSearchDebugger.cpp b/Engine/Plugins/Animation/PoseSearch/Source/Editor/Private/PoseSearchDebugger.cpp
--- a/Engine/Plugins/Animation/PoseSearch/Source/Editor/Private/PoseSearchDebugger.cpp
+++ b/Engine/Plugins/Animation/PoseSearch/Source/Editor/Private/PoseSearchDebugger.cpp
@@ -231,7 +231,7 @@ FReply SCostTimelineView::OnMouseMove(const FGeometry& MyGeometry, const FPointe
 		const double TargetTime = RangeToScreen.LocalXToInput(HitPosition.X);
 
 		// Get curve value at given time
-		const TArray<FCurvePoint>& CurvePoints = BestCostView->CurveData->Points;
+		const TConstArrayView<FCurvePoint> CurvePoints = BestCostView->CurveData->Points;
 		const int32 NumPoints = CurvePoints.Num();
 
 		if (NumPoints > 0)
diff --git a/Engine/Plugins/Animation/PoseSearch/Source/Editor/Private/PoseSearchDebuggerDatabaseView.cpp b/Engine/Plugins/Animation/PoseSearch/Source/Editor/Private/PoseSearchDebuggerDatabaseView.cpp
--- a/Engine/Plugins/Animation/PoseSearch/Source/Editor/Private/PoseSearchDebuggerDatabaseView.cpp
+++ b/Engine/Plugins/Animation/PoseSearch/Source/Editor/Private/PoseSearchDebuggerDatabaseView.cpp
@@ -87,7 +87,7 @@ public:
 		}
 	}
 
-	bool AreLabelsEqualTo(const TArray<FText>& OtherLabels) const
+	bool AreLabelsEqualTo(const TConstArrayView<FText> OtherLabels) const
 	{
 		if (LabelToChannels.Num() != OtherLabels.Num())
 		{
diff --git a/Engine/Plugins/Animation/PoseSearch/Source/Editor/Public/PoseSearchInteractionAssetEditor.h b/Engine/Plugins/Animation/PoseSearch/Source/Editor/Public/PoseSearchInteractionAssetEditor.h
--- a/Engine/Plugins/Animation/PoseSearch/Source/Editor/Public/PoseSearchInteractionAssetEditor.h
+++ b/Engine/Plugins/Animation/PoseSearch/Source/Editor/Public/PoseSearchInteractionAssetEditor.h
@@ -66,8 +66,7 @@ public:
 	void PreviewForwardEnd();
 	const UPoseSearchInteractionAsset* GetInteractionAsset() const { return InteractionAssetPtr.Get(); }
 	void Tick(float DeltaSeconds);
-	const TArray<FInteractionAssetPreviewActor>& GetPreviewActors() const { return PreviewActors; }
-	TArray<FInteractionAssetPreviewActor>& GetPreviewActors() { return PreviewActors; }
+	TConstArrayView<FInteractionAssetPreviewActor> GetPreviewActors() const { return PreviewActors; }
 	TRange<double> GetPreviewPlayRange() const;
 	void SetPlayTime(float NewPlayTime, bool bInTickPlayTime);
 	float GetPlayTime() const { return PlayTime; }
diff --git a/Engine/Plugins/Animation/PoseSearch/Source/Runtime/Private/PoseSearchDerivedData.cpp b/Engine/Plugins/Animation/PoseSearch/Source/Runtime/Private/PoseSearchDerivedData.cpp
--- a/Engine/Plugins/Animation/PoseSearch/Source/Runtime/Private/PoseSearchDerivedData.cpp
+++ b/Engine/Plugins/Animation/PoseSearch/Source/Runtime/Private/PoseSearchDerivedData.cpp
@@ -12,7 +12,7 @@
 #include "Chooser/Internal/Chooser.h"
 #include "DerivedDataCache.h"
 #include "DerivedDataRequestOwner.h"
-#include "StructUtils/InstancedStruct.h"
+#include "Interfaces/ITargetPlatformManagerModule.h"
 #include "Misc/CoreDelegates.h"
 #if ENABLE_ANIM_DEBUG
 #include "Misc/FileHelper.h"
@@ -29,9 +29,9 @@
 #include "ProfilingDebugging/CookStats.h"
 #include "ScopedTransaction.h"
 #include "Serialization/BulkDataRegistry.h"
+#include "StructUtils/InstancedStruct.h"
 #include "UObject/NoExportTypes.h"
 #include "UObject/PackageReload.h"
-#include "Interfaces/ITargetPlatformManagerModule.h"
 
 #define LOCTEXT_NAMESPACE "PoseSearchDerivedData"
 
@@ -1946,7 +1946,7 @@ void FPoseSearchDatabaseAsyncCacheTask::OnGetComplete(UE::DerivedData::FCacheGet
 								}
 								else
 								{
-									const TArray<FBlendSample>& BlendSamples = BlendSpace->GetBlendSamples();
+									const TConstArrayView<FBlendSample> BlendSamples = BlendSpace->GetBlendSamples();
 									for (int i = 0; i < BlendSamples.Num() - 1; ++i)
 									{
 										const FBlendSample& CurrSample = BlendSamples[i];
diff --git a/Engine/Plugins/Animation/PoseSearch/Source/Runtime/Public/PoseSearch/PoseSearchContext.h b/Engine/Plugins/Animation/PoseSearch/Source/Runtime/Public/PoseSearch/PoseSearchContext.h
--- a/Engine/Plugins/Animation/PoseSearch/Source/Runtime/Public/PoseSearch/PoseSearchContext.h
+++ b/Engine/Plugins/Animation/PoseSearch/Source/Runtime/Public/PoseSearch/PoseSearchContext.h
@@ -199,7 +199,7 @@ struct FSearchContext
 	const FPoseIndicesHistory* GetPoseIndicesHistory() const { return PoseIndicesHistory; }
 	
 	UE_API bool ArePoseHistoriesValid() const;
-	const TArray<const IPoseHistory*, TInlineAllocator<PreallocatedRolesNum, TMemStackAllocator<>>>& GetPoseHistories() const { return PoseHistories; }
+	TConstArrayView<const IPoseHistory*> GetPoseHistories() const { return PoseHistories; }
 	UE_API const IPoseHistory* GetPoseHistory(const FRole& Role) const;
 
 	float GetDesiredPermutationTimeOffset() const { return DesiredPermutationTimeOffset; }
diff --git a/Engine/Plugins/Experimental/AnimNextPoseSearch/Source/AnimNextPoseSearch/Private/MotionMatchingTrait.cpp b/Engine/Plugins/Experimental/AnimNextPoseSearch/Source/AnimNextPoseSearch/Private/MotionMatchingTrait.cpp
--- a/Engine/Plugins/Experimental/AnimNextPoseSearch/Source/AnimNextPoseSearch/Private/MotionMatchingTrait.cpp
+++ b/Engine/Plugins/Experimental/AnimNextPoseSearch/Source/AnimNextPoseSearch/Private/MotionMatchingTrait.cpp
@@ -187,7 +187,7 @@ namespace UE::AnimNext
 #endif
 
 		bool bIsInteraction = false;
-		const TArray<FPoseSearchInteractionAvailability>& Availabilities = SharedData->GetAvailabilities(Binding);
+		const TArrayView<const FPoseSearchInteractionAvailability> Availabilities = SharedData->GetAvailabilities(Binding);
 		if (!Availabilities.IsEmpty())
 		{
 			// WIP motion matched interactions
