BlendStack - storing FBlendStackAnimPlayer::StoredAttributes as FMeshAttributeContainer instead of FHeapAttributeContainer

diff --git a/Engine/Plugins/Animation/BlendStack/Source/Runtime/Private/AnimNode_BlendStack.cpp b/Engine/Plugins/Animation/BlendStack/Source/Runtime/Private/AnimNode_BlendStack.cpp
--- a/Engine/Plugins/Animation/BlendStack/Source/Runtime/Private/AnimNode_BlendStack.cpp
+++ b/Engine/Plugins/Animation/BlendStack/Source/Runtime/Private/AnimNode_BlendStack.cpp
@@ -164,7 +164,7 @@ void FBlendStackAnimPlayer::StorePoseContext(const FPoseContext& PoseContext)
 	}
 
 	StoredCurve.CopyFrom(PoseContext.Curve);
-	StoredAttributes.CopyFrom(PoseContext.CustomAttributes);
+	StoredAttributes.CopyFrom(PoseContext.CustomAttributes, PoseContext.Pose.GetBoneContainer());
 }
 
 bool FBlendStackAnimPlayer::HasValidPoseContext() const
@@ -230,10 +230,9 @@ void FBlendStackAnimPlayer::RestorePoseContext(FPoseContext& PoseContext) const
 	}
 	
 	PoseContext.Curve.CopyFrom(StoredCurve);
-	PoseContext.CustomAttributes.CopyFrom(StoredAttributes);
+	PoseContext.CustomAttributes.CopyFrom(StoredAttributes, PoseContext.Pose.GetBoneContainer());
 }
 
-
 // @todo: maybe implement copy/move constructors and assignment operator do so (or use a list instead of an array)
 // since we're making copies and moving this object in memory, we're using this method to set the MirrorNode SourceLinkNode when necessary
 void FBlendStackAnimPlayer::UpdateSourceLinkNode()
diff --git a/Engine/Plugins/Animation/BlendStack/Source/Runtime/Public/BlendStack/AnimNode_BlendStack.h b/Engine/Plugins/Animation/BlendStack/Source/Runtime/Public/BlendStack/AnimNode_BlendStack.h
--- a/Engine/Plugins/Animation/BlendStack/Source/Runtime/Public/BlendStack/AnimNode_BlendStack.h
+++ b/Engine/Plugins/Animation/BlendStack/Source/Runtime/Public/BlendStack/AnimNode_BlendStack.h
@@ -99,7 +99,7 @@ private:
 	// NoTe: we don't need a full FCompactHeapPose, since we use StoredBoneContainer to cache the bone container
 	TArray<FTransform> StoredBones;
 	FBlendedHeapCurve StoredCurve;
-	UE::Anim::FHeapAttributeContainer StoredAttributes;
+	UE::Anim::FMeshAttributeContainer StoredAttributes;
 	// We need to store the bone container, in case we have a LOD swap during a blend that uses the stored pose.
 	FBoneContainer StoredBoneContainer;
 
