Fix for Blend Stack Input IsActive not working as expected
#rnx

diff --git a/Engine/Plugins/Animation/BlendStack/Source/Runtime/Private/AnimNode_BlendStack.cpp b/Engine/Plugins/Animation/BlendStack/Source/Runtime/Private/AnimNode_BlendStack.cpp
--- a/Engine/Plugins/Animation/BlendStack/Source/Runtime/Private/AnimNode_BlendStack.cpp
+++ b/Engine/Plugins/Animation/BlendStack/Source/Runtime/Private/AnimNode_BlendStack.cpp
@@ -352,11 +352,6 @@ float FBlendStackAnimPlayer::GetPlayRate() const
 	return 0.f;
 }
 
-bool FBlendStackAnimPlayer::IsActive() const
-{
-	return TimeToActivation <= 0.f;
-}
-
 FAnimNode_AssetPlayerBase* FBlendStackAnimPlayer::GetAssetPlayerNode()
 {
 	if (SequencePlayerNode.GetSequence())
@@ -603,7 +598,7 @@ void FAnimNode_BlendStack_Standalone::Evaluate_AnyThread(FPoseContext& Output)
 
 		for (int32 i = 0; i < AnimPlayers.Num(); ++i)
 		{
-			if (AnimPlayers[i].IsActive())
+			if (!AnimPlayers[i].IsPostponed())
 			{
 				FString IndexString = FString("[" + FString::FromInt(i) + FString("]"));
 				FString Asset = FString("Asset") + IndexString;
@@ -785,6 +780,7 @@ void FAnimNode_BlendStack_Standalone::UpdateAssetPlayer(const FAnimationUpdateCo
 	float CurrentWeightMultiplier = 1.f;
 	const int32 BlendStackSize = AnimPlayers.Num();
 	int32 AnimPlayerIndex = 0;
+	bool bLocalIsContextActive = true;
 	while (AnimPlayerIndex < BlendStackSize)
 	{
 		FBlendStackAnimPlayer& AnimPlayer = AnimPlayers[AnimPlayerIndex];
@@ -798,7 +794,14 @@ void FAnimNode_BlendStack_Standalone::UpdateAssetPlayer(const FAnimationUpdateCo
 #if ENABLE_ANIM_DEBUG || ENABLE_VISUAL_LOG 
 			UE::Anim::TScopedGraphMessage<UE::Anim::FAnimBlendDebugScope> BlendDebugMessage(Context, Context, AnimPlayerIndex, BlendStackSize, AnimPlayer.GetDebugColor());
 #endif
-			UpdateSample((AnimPlayerIndex == 0) ? AnimPlayerContext : AnimPlayerContext.AsInactive(), AnimPlayerIndex);
+
+			// A postponed anim is considered inactive since it hasn't started blending in yet.
+			const bool bIsPlayerContextActive = bLocalIsContextActive && Context.IsActive() && !AnimPlayer.IsPostponed();
+			AnimPlayer.bIsActive = bIsPlayerContextActive;
+			UpdateSample(bIsPlayerContextActive ? AnimPlayerContext : AnimPlayerContext.AsInactive(), AnimPlayerIndex);
+
+			// Found our active context! All further ones should be inactive.
+			bLocalIsContextActive &= !bIsPlayerContextActive;
 		}
 		CurrentWeightMultiplier *= (1.f - BlendInPercentage);
 
@@ -832,7 +835,7 @@ bool FAnimNode_BlendStack_Standalone::EvaluateSample(FPoseContext& Output, const
 {
 	FBlendStackAnimPlayer& SamplePlayer = AnimPlayers[PlayerIndex];
 	
-	if (!SamplePlayer.IsActive())
+	if (SamplePlayer.IsPostponed())
 	{
 		return false;
 	}
@@ -883,7 +886,7 @@ void FAnimNode_BlendStack_Standalone::UpdateSample(const FAnimationUpdateContext
 	// Advance the blend-in time regardless of whether or not the player was updated.
 	SamplePlayer.UpdateWithDeltaTime(Context.GetDeltaTime(), PlayerIndex, PlayerDepthBlendInTimeMultiplier);
 
-	if (SamplePlayer.IsActive())
+	if (!SamplePlayer.IsPostponed())
 	{
 		const bool bHasSampleGraph = IsSampleGraphAvailableForPlayer(PlayerIndex);
 		if (bHasSampleGraph)
@@ -1083,7 +1086,7 @@ void FAnimNode_BlendStack_Standalone::InternalBlendTo(const FAnimationUpdateCont
 	}
 
 	int32 NewSamplePoseLinkIndex = CurrentSamplePoseLink;
-	if (!bBlendStackIsEmpty && !AnimPlayers[0].IsActive())
+	if (!bBlendStackIsEmpty && AnimPlayers[0].IsPostponed())
 	{
 		// we allow only one player with TimeToActivation > 0:
 		// replacing AnimPlayers[0] with this new BlendTo request
diff --git a/Engine/Plugins/Animation/BlendStack/Source/Runtime/Public/BlendStack/AnimNode_BlendStack.h b/Engine/Plugins/Animation/BlendStack/Source/Runtime/Public/BlendStack/AnimNode_BlendStack.h
--- a/Engine/Plugins/Animation/BlendStack/Source/Runtime/Public/BlendStack/AnimNode_BlendStack.h
+++ b/Engine/Plugins/Animation/BlendStack/Source/Runtime/Public/BlendStack/AnimNode_BlendStack.h
@@ -69,7 +69,11 @@ struct FBlendStackAnimPlayer
 	UE_API void RestorePoseContext(FPoseContext& PoseContext) const;
 	UE_API void UpdateSourceLinkNode();
 	UE_API bool IsLooping() const;
-	UE_API bool IsActive() const;
+	UE_API bool IsPostponed() const { return TimeToActivation > 0.f; };
+
+	// A sample being "Active" is relative to the active branch of the animation update context
+	// i.e. Active branches are NOT blending out.
+	UE_API bool IsActive() const { return bIsActive; };
 
 	UE_API FAnimNode_AssetPlayerBase* GetAssetPlayerNode();
 
@@ -110,6 +114,7 @@ private:
 	float TotalBlendInTime = 0.f;
 	float CurrentBlendInTime = 0.f;
 	float TimeToActivation = 0.f;
+	bool bIsActive = false;
 
 #if ENABLE_ANIM_DEBUG || ENABLE_VISUAL_LOG 
 	FColor DebugColor = FColor::Black;
