BlendStack - logging when the leaving animation remaining time is less than the new animation blend in time in FAnimNode_BlendStack_Standalone::InternalBlendTo, since it'd cause poor quality blends

#rnx

diff --git a/Engine/Plugins/Animation/BlendStack/Source/Runtime/Private/AnimNode_BlendStack.cpp b/Engine/Plugins/Animation/BlendStack/Source/Runtime/Private/AnimNode_BlendStack.cpp
--- a/Engine/Plugins/Animation/BlendStack/Source/Runtime/Private/AnimNode_BlendStack.cpp
+++ b/Engine/Plugins/Animation/BlendStack/Source/Runtime/Private/AnimNode_BlendStack.cpp
@@ -38,7 +38,7 @@ void FBlendStackAnimPlayer::Initialize(const FAnimationInitializeContext& Contex
 {
 	if (bMirrored && !MirrorDataTable)
 	{
-		UE_LOG(LogBlendStack, Error, TEXT("FBlendStackAnimPlayer failed to Initialize for %s. Mirroring will not work becasue MirrorDataTable is missing"), *GetNameSafe(AnimationAsset));
+		UE_LOG(LogBlendStack, Error, TEXT("FBlendStackAnimPlayer::Initialize failed for '%s'. Mirroring will not work becasue MirrorDataTable is missing"), *GetNameSafe(AnimationAsset));
 	}
 	
 	check(Context.AnimInstanceProxy);
@@ -49,7 +49,7 @@ void FBlendStackAnimPlayer::Initialize(const FAnimationInitializeContext& Contex
 	const int32 NumSkeletonBones = RefSkeleton.GetNum();
 	if (NumSkeletonBones <= 0)
 	{
-		UE_LOG(LogBlendStack, Error, TEXT("FBlendStackAnimPlayer failed to Initialize for %s. Skeleton has no bones?!"), *GetNameSafe(AnimationAsset));
+		UE_LOG(LogBlendStack, Error, TEXT("FBlendStackAnimPlayer::Initialize failed for '%s'. Skeleton has no bones?!"), *GetNameSafe(AnimationAsset));
 	}
 	else if (BlendTime > UE_KINDA_SMALL_NUMBER)
 	{
@@ -116,7 +116,7 @@ void FBlendStackAnimPlayer::Initialize(const FAnimationInitializeContext& Contex
 		BlendSpacePlayerNode.SetBlendSpace(nullptr);
 		SequencePlayerNode.SetSequence(nullptr);
 
-		UE_LOG(LogBlendStack, Error, TEXT("FBlendStackAnimPlayer unsupported AnimationAsset %s"), *GetNameSafe(AnimationAsset));
+		UE_LOG(LogBlendStack, Error, TEXT("FBlendStackAnimPlayer::Initialize unsupported AnimationAsset %s"), *GetNameSafe(AnimationAsset));
 	}
 
 	UpdateSourceLinkNode();
@@ -139,6 +139,29 @@ void FBlendStackAnimPlayer::Initialize(const FAnimationInitializeContext& Contex
 #endif
 }
 
+float FBlendStackAnimPlayer::GetSecondsFromEnd() const
+{
+	float TimeFromEndInSeconds = 0.f;
+
+	if (SequencePlayerNode.GetSequence())
+	{
+		const float AccumulatedTime = SequencePlayerNode.GetAccumulatedTime();
+		TimeFromEndInSeconds = SequencePlayerNode.GetTimeFromEnd(AccumulatedTime);
+	}
+	else if (const UBlendSpace* BlendSpace = BlendSpacePlayerNode.GetBlendSpace())
+	{
+		int32 TriangulationIndex = 0;
+		TArray<FBlendSampleData> BlendSamples;
+		const FVector BlendParameters = BlendSpacePlayerNode.GetPosition();
+		BlendSpace->GetSamplesFromBlendInput(BlendParameters, BlendSamples, TriangulationIndex, true);
+		const float PlayLength = BlendSpace->GetAnimationLengthFromSampleData(BlendSamples);
+		const float AccumulatedTime = BlendSpacePlayerNode.GetAccumulatedTime();
+		TimeFromEndInSeconds = BlendSpacePlayerNode.GetTimeFromEnd(AccumulatedTime) * PlayLength;
+	}
+
+	return TimeFromEndInSeconds;
+}
+
 void FBlendStackAnimPlayer::UpdatePlayRate(float PlayRate)
 {
 	if (SequencePlayerNode.GetSequence())
@@ -695,7 +718,7 @@ void FAnimNode_BlendStack_Standalone::Evaluate_AnyThread(FPoseContext& Output)
 			// warning if we're dropping an animplayer with relevant (MaxBlendInTimeToOverrideAnimation) weight (GetBlendInPercentage)
 			else if (AnimPlayers.Last().GetBlendInPercentage() < (1.f - MaxBlendInTimeToOverrideAnimation))
 			{
-				UE_LOG(LogBlendStack, Warning, TEXT("FAnimNode_BlendStack_Standalone dropping animplayer with blend in at %.2f"), AnimPlayers.Last().GetBlendInPercentage());
+				UE_LOG(LogBlendStack, Warning, TEXT("FAnimNode_BlendStack_Standalone::Evaluate_AnyThread dropping animplayer with blend in at %.2f"), AnimPlayers.Last().GetBlendInPercentage());
 			}
 #endif // !NO_LOGGING
 		}
@@ -1085,6 +1108,8 @@ void FAnimNode_BlendStack_Standalone::InternalBlendTo(const FAnimationUpdateCont
 		BlendTime = 0.0f;
 	}
 
+	
+
 	int32 NewSamplePoseLinkIndex = CurrentSamplePoseLink;
 	if (!bBlendStackIsEmpty && AnimPlayers[0].IsPostponed())
 	{
@@ -1096,18 +1121,40 @@ void FAnimNode_BlendStack_Standalone::InternalBlendTo(const FAnimationUpdateCont
 		AnimPlayers[0].GetCurrentBlendInTime() < MaxBlendInTimeToOverrideAnimation)
 	{
 		// replacing AnimPlayers[0] with this new BlendTo request
-		UE_LOG(LogBlendStack, Verbose, TEXT("FAnimNode_BlendStack_Standalone '%s' replaced by '%s' because blend time in is less than MaxBlendInTimeToOverrideAnimation (%.2f / %.2f)"), *AnimPlayers[0].GetAnimationName(), *GetNameSafe(AnimationAsset), AnimPlayers[0].GetCurrentBlendInTime(), MaxBlendInTimeToOverrideAnimation);
+		UE_LOG(LogBlendStack, Verbose, TEXT("FAnimNode_BlendStack_Standalone::InternalBlendTo '%s' replaced by '%s' because blend time in is less than MaxBlendInTimeToOverrideAnimation (%.2f / %.2f)"), *AnimPlayers[0].GetAnimationName(), *GetNameSafe(AnimationAsset), AnimPlayers[0].GetCurrentBlendInTime(), MaxBlendInTimeToOverrideAnimation);
 	}
 	else if (AnimPlayers.Num() <= MaxActiveBlends + 2)
 	{
+#if !NO_LOGGING
+		if (!bBlendStackIsEmpty)
+		{
+			const float SecondsFromEnd = AnimPlayers[0].GetSecondsFromEnd();
+			if (SecondsFromEnd < BlendTime)
+			{
+				UE_LOG(LogBlendStack, Verbose, TEXT("FAnimNode_BlendStack_Standalone::InternalBlendTo remaining playback time for '%s' is only %f sec, while the requested blend time is %f sec. This will result in a poor quality blend, since the currently playing asset will end and freeze its pose, curves, and attributes before the new animation is fully relevant"), *AnimPlayers[0].GetAnimationName(), SecondsFromEnd, BlendTime);
+			}
+		}
+#endif // !NO_LOGGING
+
 		AnimPlayers.Insert(FBlendStackAnimPlayer(), 0);
 		// ...otherwise, assign a new graph.
 		NewSamplePoseLinkIndex = GetNextPoseLinkIndex();
 	}
 	else
 	{
+#if !NO_LOGGING
+		if (AnimPlayers.Num() > 1)
+		{
+			const float SecondsFromEnd = AnimPlayers[1].GetSecondsFromEnd();
+			if (SecondsFromEnd < BlendTime)
+			{
+				UE_LOG(LogBlendStack, Verbose, TEXT("FAnimNode_BlendStack_Standalone::InternalBlendTo remaining playback time for '%s' is only %f sec, while the requested blend time is %f sec. This will result in a poor quality blend, since the currently playing asset will end and freeze its pose, curves, and attributes before the new animation is fully relevant"), *AnimPlayers[0].GetAnimationName(), SecondsFromEnd, BlendTime);
+			}
+		}
+#endif // !NO_LOGGING
+
 		// else it means we had multiple BlendTo suring the same frame. we'll let the last one win
-		UE_LOG(LogBlendStack, Warning, TEXT("FAnimNode_BlendStack_Standalone multiple BlendTo requests during the same frame: only the last request will be put on this BlendStack"));
+		UE_LOG(LogBlendStack, Warning, TEXT("FAnimNode_BlendStack_Standalone::InternalBlendTo multiple BlendTo requests during the same frame: only the last request will be put on this BlendStack"));
 	}
 
 	FBlendStackAnimPlayer& AnimPlayer = AnimPlayers[0];
diff --git a/Engine/Plugins/Animation/BlendStack/Source/Runtime/Public/BlendStack/AnimNode_BlendStack.h b/Engine/Plugins/Animation/BlendStack/Source/Runtime/Public/BlendStack/AnimNode_BlendStack.h
--- a/Engine/Plugins/Animation/BlendStack/Source/Runtime/Public/BlendStack/AnimNode_BlendStack.h
+++ b/Engine/Plugins/Animation/BlendStack/Source/Runtime/Public/BlendStack/AnimNode_BlendStack.h
@@ -76,6 +76,8 @@ struct FBlendStackAnimPlayer
 	bool IsActive() const { return bIsActive; };
 
 	UE_API FAnimNode_AssetPlayerBase* GetAssetPlayerNode();
+	
+	float GetSecondsFromEnd() const;
 
 #if ENABLE_ANIM_DEBUG || ENABLE_VISUAL_LOG 
 	UE_API FColor GetDebugColor() const;
