PoseSearch - fix for misspelled words

#rnx
#Synced-CL 49689419

diff --git a/Engine/Plugins/Animation/BlendStack/Source/Runtime/Private/AnimNode_BlendStack.cpp b/Engine/Plugins/Animation/BlendStack/Source/Runtime/Private/AnimNode_BlendStack.cpp
--- a/Engine/Plugins/Animation/BlendStack/Source/Runtime/Private/AnimNode_BlendStack.cpp
+++ b/Engine/Plugins/Animation/BlendStack/Source/Runtime/Private/AnimNode_BlendStack.cpp
@@ -38,7 +38,7 @@ void FBlendStackAnimPlayer::Initialize(const FAnimationInitializeContext& Contex
 {
 	if (bMirrored && !MirrorDataTable)
 	{
-		UE_LOG(LogBlendStack, Error, TEXT("FBlendStackAnimPlayer::Initialize failed for '%s'. Mirroring will not work becasue MirrorDataTable is missing"), *GetNameSafe(AnimationAsset));
+		UE_LOG(LogBlendStack, Error, TEXT("FBlendStackAnimPlayer::Initialize failed for '%s'. Mirroring will not work because MirrorDataTable is missing"), *GetNameSafe(AnimationAsset));
 	}
 	
 	check(Context.AnimInstanceProxy);
@@ -838,7 +838,7 @@ void FAnimNode_BlendStack_Standalone::UpdateAssetPlayer(const FAnimationUpdateCo
 
 	// AnimPlayers[AnimPlayerIndex] is the first FBlendStackAnimPlayer with a weight contribution of zero,
 	// so we can discard it and all the successive AnimPlayers as well
-	// NoTe that it's safe to delete all those players, becasue we didn't call SamplePlayer.Update_AnyThread 
+	// NoTe that it's safe to delete all those players, because we didn't call SamplePlayer.Update_AnyThread 
 	// hence not register sequence / blendspace player InternalTimeAccumulator via FAnimTickRecord(s)
 	const int32 WantedAnimPlayersNum = FMath::Max(1, AnimPlayerIndex); // we save at least one FBlendStackAnimPlayer
 	while (AnimPlayers.Num() > WantedAnimPlayersNum)
diff --git a/Engine/Plugins/Animation/PoseSearch/Source/Runtime/Private/PoseSearchDerivedData.cpp b/Engine/Plugins/Animation/PoseSearch/Source/Runtime/Private/PoseSearchDerivedData.cpp
--- a/Engine/Plugins/Animation/PoseSearch/Source/Runtime/Private/PoseSearchDerivedData.cpp
+++ b/Engine/Plugins/Animation/PoseSearch/Source/Runtime/Private/PoseSearchDerivedData.cpp
@@ -1992,7 +1992,7 @@ void FPoseSearchDatabaseAsyncCacheTask::OnGetComplete(UE::DerivedData::FCacheGet
 					// Building all the related FPoseSearchBaseIndex first
 					if (!InitSearchIndexAssets(SearchIndexBase, DependentDatabase, DependentDatabaseSchema, DependentExcludeFromDatabaseParameters))
 					{
-						UE_LOG(LogPoseSearch, Error, TEXT("%s - %s BuildIndex Failed becasue of invalid assets"), *LexToString(FullIndexKey.Hash), *MainDatabaseName);
+						UE_LOG(LogPoseSearch, Error, TEXT("%s - %s BuildIndex Failed because of invalid assets"), *LexToString(FullIndexKey.Hash), *MainDatabaseName);
 						ResetSearchIndex();
 						return;
 					}
diff --git a/Engine/Plugins/Animation/PoseSearch/Source/Runtime/Private/PoseSearchLibrary.cpp b/Engine/Plugins/Animation/PoseSearch/Source/Runtime/Private/PoseSearchLibrary.cpp
--- a/Engine/Plugins/Animation/PoseSearch/Source/Runtime/Private/PoseSearchLibrary.cpp
+++ b/Engine/Plugins/Animation/PoseSearch/Source/Runtime/Private/PoseSearchLibrary.cpp
@@ -670,7 +670,7 @@ void UPoseSearchLibrary::TraceMotionMatching(
 	}
 
 	// finalizing SearchContext.GetBestPoseCandidatesMap() by calling SearchContext::Track for all the SearchResults flagging them as EPoseCandidateFlags::Valid_CurrentPose
-	// because, those candidate could have been discarded becasue of their cost was too high and didn't fit into SearchContext.GetBestPoseCandidatesMap() limits
+	// because, those candidate could have been discarded because of their cost was too high and didn't fit into SearchContext.GetBestPoseCandidatesMap() limits
 	SearchResults.IterateOverSearchResults([&SearchContext](FSearchResult& SearchResult)
 		{
 			if (SearchResult.IsValid())
