BlendStack - sync group support for follower blendstacks

#Synced-CL 49959264

diff --git a/Engine/Plugins/Animation/BlendStack/Source/Runtime/Private/AnimNode_BlendStack.cpp b/Engine/Plugins/Animation/BlendStack/Source/Runtime/Private/AnimNode_BlendStack.cpp
--- a/Engine/Plugins/Animation/BlendStack/Source/Runtime/Private/AnimNode_BlendStack.cpp
+++ b/Engine/Plugins/Animation/BlendStack/Source/Runtime/Private/AnimNode_BlendStack.cpp
@@ -33,7 +33,7 @@ namespace UE::BlendStack
 /////////////////////////////////////////////////////
 // FBlendStackAnimPlayer
 void FBlendStackAnimPlayer::Initialize(const FAnimationInitializeContext& Context, UAnimationAsset* AnimationAsset, float AccumulatedTime, bool bLoop,
-	bool bMirrored, UMirrorDataTable* MirrorDataTable, float BlendTime, const UBlendProfile* BlendProfile, EAlphaBlendOption InBlendOption,
+	bool bMirrored, UMirrorDataTable* MirrorDataTable, float BlendTime, const UBlendProfile* InBlendProfile, EAlphaBlendOption InBlendOption,
 	const FVector& BlendParameters, float PlayRate, float ActivationDelay, int32 InPoseLinkIdx, FName GroupName, EAnimGroupRole::Type GroupRole, EAnimSyncMethod GroupMethod, bool bOverridePositionWhenJoiningSyncGroupAsLeader)
 {
 	if (bMirrored && !MirrorDataTable)
@@ -54,14 +54,15 @@ void FBlendStackAnimPlayer::Initialize(const FAnimationInitializeContext& Contex
 	else if (BlendTime > UE_KINDA_SMALL_NUMBER)
 	{
 		// handling BlendTime > 0 and RootBoneBlendTime >= 0
-		if (BlendProfile != nullptr)
+		if (InBlendProfile != nullptr)
 		{
 			TotalBlendInTimePerBone.SetNumUninitialized(NumSkeletonBones);
-			BlendProfile->FillSkeletonBoneDurationsArray(TotalBlendInTimePerBone, BlendTime, Skeleton);
+			InBlendProfile->FillSkeletonBoneDurationsArray(TotalBlendInTimePerBone, BlendTime, Skeleton);
 		}
 	}
 
 	BlendOption = InBlendOption;
+	BlendProfile = InBlendProfile;
 
 	TotalBlendInTime = BlendTime;
 	CurrentBlendInTime = 0.f;
@@ -330,6 +331,76 @@ float FBlendStackAnimPlayer::GetAccumulatedTime() const
 	return 0.f;
 }
 
+EAnimSyncMethod FBlendStackAnimPlayer::GetGroupMethod() const
+{
+	if (SequencePlayerNode.GetSequence())
+	{
+		return SequencePlayerNode.GetGroupMethod();
+	}
+	
+	if (BlendSpacePlayerNode.GetBlendSpace())
+	{
+		return BlendSpacePlayerNode.GetGroupMethod();
+	}
+
+	return EAnimSyncMethod::DoNotSync;
+}
+
+FName FBlendStackAnimPlayer::GetGroupName() const
+{
+	if (SequencePlayerNode.GetSequence())
+	{
+		return SequencePlayerNode.GetGroupName();
+	}
+	
+	if (BlendSpacePlayerNode.GetBlendSpace())
+	{
+		return BlendSpacePlayerNode.GetGroupName();
+	}
+
+	return FName();
+}
+
+EAnimGroupRole::Type FBlendStackAnimPlayer::GetGroupRole() const
+{
+	if (SequencePlayerNode.GetSequence())
+	{
+		return SequencePlayerNode.GetGroupRole();
+	}
+	
+	if (BlendSpacePlayerNode.GetBlendSpace())
+	{
+		return BlendSpacePlayerNode.GetGroupRole();
+	}
+
+	return EAnimGroupRole::Type::CanBeLeader;
+}
+
+bool FBlendStackAnimPlayer::GetOverridePositionWhenJoiningSyncGroupAsLeader() const
+{
+	if (SequencePlayerNode.GetSequence())
+	{
+		return SequencePlayerNode.GetOverridePositionWhenJoiningSyncGroupAsLeader();
+	}
+	
+	if (BlendSpacePlayerNode.GetBlendSpace())
+	{
+		return BlendSpacePlayerNode.GetOverridePositionWhenJoiningSyncGroupAsLeader();
+	}
+
+	return false;
+}
+
+const UBlendProfile* FBlendStackAnimPlayer::GetBlendProfile() const
+{
+	return BlendProfile.Get();
+}
+
+UMirrorDataTable* FBlendStackAnimPlayer::GetMirrorDataTable() const
+{
+	return MirrorNode.GetMirrorDataTable();
+}
+
 float FBlendStackAnimPlayer::GetCurrentAssetTime() const
 {
 	if (SequencePlayerNode.GetSequence())
@@ -796,7 +867,7 @@ void FAnimNode_BlendStack_Standalone::UpdateAssetPlayer(const FAnimationUpdateCo
 	}
 	
 	UE::Anim::TOptionalScopedGraphMessage<UE::Anim::FBlendStackAnimEventsFilterScope> BlendStackInfoScope(bShouldFilterNotifies, Context, NotifiesFiredLastTick, NotifyRecencyMap);
-
+		
 	// AnimPlayers[0] is the most newly inserted AnimPlayer, AnimPlayers[AnimPlayers.Num()-1] is the oldest, so to calculate the weights
 	// we ask AnimPlayers[0] its BlendInPercentage and then distribute the left over (CurrentWeightMultiplier) to the rest of the AnimPlayers
 	// AnimPlayers[AnimPlayerIndex].GetBlendWeight() will now store the weighted contribution of AnimPlayers[AnimPlayerIndex] to be able to calculate root motion from animation
@@ -1017,82 +1088,92 @@ static void RequestInertialBlend(const FAnimationUpdateContext& Context, float B
 	}
 }
 
-void FAnimNode_BlendStack_Standalone::BlendTo(const FAnimationUpdateContext& Context, UAnimationAsset* AnimationAsset, float AccumulatedTime, bool bLoop,
+bool FAnimNode_BlendStack_Standalone::StitchTo(const FAnimationUpdateContext& Context, UAnimationAsset* AnimationAsset, float AccumulatedTime, bool bLoop,
 	bool bMirrored, UMirrorDataTable* MirrorDataTable, float BlendTime, const UBlendProfile* BlendProfile, EAlphaBlendOption BlendOption, bool bUseInertialBlend, FName InertialBlendNodeTag,
-	const FVector& BlendParameters, float PlayRate, float ActivationDelay, FName GroupName, EAnimGroupRole::Type GroupRole, EAnimSyncMethod GroupMethod, bool bOverridePositionWhenJoiningSyncGroupAsLeader)
+	const FVector& BlendParameters, float PlayRate, float ActivationDelay)
 {
 	using namespace UE::Anim;
 
-	bool bNeedToBlendTo = true;
-	if (StitchDatabase)
+	if (!StitchDatabase)
 	{
-		if (MaxActiveBlends > 0)
-		{
-			if (IPoseSearchProvider* PoseSearchProvider = IPoseSearchProvider::Get())
-			{
-				// looking for an animation stitch from the StitchDatabase that will connect, in BlendTime seconds, 
-				// the currently playing animation pose to the pose from AnimationAsset at AccumulatedTime + BlendTime
-				const UObject* AssetToSearch = StitchDatabase.Get();
-
-				IPoseSearchProvider::FSearchPlayingAsset PlayingAsset;
-				PlayingAsset.Asset = GetAnimAsset();
-				PlayingAsset.AccumulatedTime = GetAccumulatedTime();
-				PlayingAsset.bMirrored = GetMirror();
-				PlayingAsset.BlendParameters = GetBlendParameters();
-
-				IPoseSearchProvider::FSearchFutureAsset FutureAsset;
-				FutureAsset.Asset = AnimationAsset;
-				FutureAsset.AccumulatedTime = AccumulatedTime + BlendTime;
-				FutureAsset.IntervalTime = BlendTime;
-
-				const IPoseSearchProvider::FSearchResult SearchResult = PoseSearchProvider->Search(Context, MakeArrayView(&AssetToSearch, 1), PlayingAsset, FutureAsset);
-				if (UAnimationAsset* StitchAnimationAsset = Cast<UAnimationAsset>(SearchResult.SelectedAsset))
-				{
-					if (SearchResult.Dissimilarity <= StitchBlendMaxCost)
-					{
-						// blend to the selected animation stitch
-						InternalBlendTo(Context, StitchAnimationAsset, SearchResult.TimeOffsetSeconds, false, SearchResult.bMirrored, MirrorDataTable,
-							StitchBlendTime, BlendProfile, BlendOption, bUseInertialBlend, InertialBlendNodeTag, BlendParameters, SearchResult.WantedPlayRate, ActivationDelay, GroupName, GroupRole, GroupMethod, bOverridePositionWhenJoiningSyncGroupAsLeader);
-
-						// blend with an ActivationDelay of BlendTime - StitchBlendTime + ActivationDelay seconds
-						// to the AnimationAsset at AccumulatedTime + BlendTime - StitchBlendTime seconds in the future,
-						// so at BlendTime seconds ahead the AnimationAsset is playing the fully blended in pose at AccumulatedTime + BlendTime
-						InternalBlendTo(Context, AnimationAsset, AccumulatedTime + BlendTime - StitchBlendTime, bLoop, bMirrored, MirrorDataTable,
-							StitchBlendTime, BlendProfile, BlendOption, bUseInertialBlend, InertialBlendNodeTag, BlendParameters, PlayRate, BlendTime - StitchBlendTime + ActivationDelay, GroupName, GroupRole, GroupMethod, bOverridePositionWhenJoiningSyncGroupAsLeader);
-
-						bNeedToBlendTo = false;
-					}
-					else
-					{
-						UE_LOG(LogBlendStack, Display, TEXT("FAnimNode_BlendStack_Standalone::BlendTo StitchDatabase '%s' search cost is %f, above StitchBlendMaxCost %f. Defaulting to regular blend"), *GetNameSafe(StitchDatabase), SearchResult.Dissimilarity, StitchBlendMaxCost);
-					}
-				}
-				else
-				{
-					UE_LOG(LogBlendStack, Error, TEXT("FAnimNode_BlendStack_Standalone::BlendTo cannot use StitchDatabase '%s', because of missing IPoseSearchProvider::Search couldn't select a StitchAnimationAsset. Defaulting to regular blend"), *GetNameSafe(StitchDatabase));
-				}
-			}
-			else
-			{
-				UE_LOG(LogBlendStack, Error, TEXT("FAnimNode_BlendStack_Standalone::BlendTo cannot use StitchDatabase '%s', because of missing IPoseSearchProvider (is PoseSearch plugin enabled?). Defaulting to regular blend"), *GetNameSafe(StitchDatabase));
-			}
-		}
-		else
-		{
-			UE_LOG(LogBlendStack, Error, TEXT("FAnimNode_BlendStack_Standalone::BlendTo cannot use StitchDatabase '%s', since MaxActiveBlends should be at least 1. Defaulting to regular blend"), *GetNameSafe(StitchDatabase));
-		}
+		return false;
+	}
+
+	if (MaxActiveBlends <= 0)
+	{
+		UE_LOG(LogBlendStack, Error, TEXT("FAnimNode_BlendStack_Standalone::BlendTo cannot use StitchDatabase '%s', since MaxActiveBlends should be at least 1. Defaulting to regular blend"), *GetNameSafe(StitchDatabase));
+		return false;
+	}
+		
+	IPoseSearchProvider* PoseSearchProvider = IPoseSearchProvider::Get();
+	if (!PoseSearchProvider)
+	{
+		UE_LOG(LogBlendStack, Error, TEXT("FAnimNode_BlendStack_Standalone::BlendTo cannot use StitchDatabase '%s', because of missing IPoseSearchProvider (is PoseSearch plugin enabled?). Defaulting to regular blend"), *GetNameSafe(StitchDatabase));
+		return false;
+	}
+
+	// looking for an animation stitch from the StitchDatabase that will connect, in BlendTime seconds, 
+	// the currently playing animation pose to the pose from AnimationAsset at AccumulatedTime + BlendTime
+	const UObject* AssetToSearch = StitchDatabase.Get();
+
+	IPoseSearchProvider::FSearchPlayingAsset PlayingAsset;
+	PlayingAsset.Asset = GetAnimAsset();
+	PlayingAsset.AccumulatedTime = GetAccumulatedTime();
+	PlayingAsset.bMirrored = GetMirror();
+	PlayingAsset.BlendParameters = GetBlendParameters();
+
+	IPoseSearchProvider::FSearchFutureAsset FutureAsset;
+	FutureAsset.Asset = AnimationAsset;
+	FutureAsset.AccumulatedTime = AccumulatedTime + BlendTime;
+	FutureAsset.IntervalTime = BlendTime;
+
+	const IPoseSearchProvider::FSearchResult SearchResult = PoseSearchProvider->Search(Context, MakeArrayView(&AssetToSearch, 1), PlayingAsset, FutureAsset);
+	UAnimationAsset* StitchAnimationAsset = Cast<UAnimationAsset>(SearchResult.SelectedAsset);
+	if (!StitchAnimationAsset)
+	{
+		UE_LOG(LogBlendStack, Error, TEXT("FAnimNode_BlendStack_Standalone::BlendTo cannot use StitchDatabase '%s', because of missing IPoseSearchProvider::Search couldn't select a StitchAnimationAsset. Defaulting to regular blend"), *GetNameSafe(StitchDatabase));
+		return false;
+	}
+
+	if (SearchResult.Dissimilarity > StitchBlendMaxCost)
+	{
+		UE_LOG(LogBlendStack, Display, TEXT("FAnimNode_BlendStack_Standalone::BlendTo StitchDatabase '%s' search cost is %f, above StitchBlendMaxCost %f. Defaulting to regular blend"), *GetNameSafe(StitchDatabase), SearchResult.Dissimilarity, StitchBlendMaxCost);
+		return false;
 	}
 
-	if (bNeedToBlendTo)
+	// blend to the selected animation stitch
+	InternalBlendTo(Context, StitchAnimationAsset, SearchResult.TimeOffsetSeconds, false, SearchResult.bMirrored, MirrorDataTable,
+		StitchBlendTime, BlendProfile, BlendOption, bUseInertialBlend, InertialBlendNodeTag, BlendParameters, SearchResult.WantedPlayRate, ActivationDelay,
+		NAME_None, EAnimGroupRole::CanBeLeader, EAnimSyncMethod::DoNotSync, false, NAME_None, EAnimGroupRole::CanBeLeader, EAnimSyncMethod::DoNotSync);
+
+	// blend with an ActivationDelay of BlendTime - StitchBlendTime + ActivationDelay seconds
+	// to the AnimationAsset at AccumulatedTime + BlendTime - StitchBlendTime seconds in the future,
+	// so at BlendTime seconds ahead the AnimationAsset is playing the fully blended in pose at AccumulatedTime + BlendTime
+	InternalBlendTo(Context, AnimationAsset, AccumulatedTime + BlendTime - StitchBlendTime, bLoop, bMirrored, MirrorDataTable,
+		StitchBlendTime, BlendProfile, BlendOption, bUseInertialBlend, InertialBlendNodeTag, BlendParameters, PlayRate, BlendTime - StitchBlendTime + ActivationDelay,
+		NAME_None, EAnimGroupRole::CanBeLeader, EAnimSyncMethod::DoNotSync, false, NAME_None, EAnimGroupRole::CanBeLeader, EAnimSyncMethod::DoNotSync);
+
+	return true;
+}
+
+void FAnimNode_BlendStack_Standalone::BlendTo(const FAnimationUpdateContext& Context, UAnimationAsset* AnimationAsset, float AccumulatedTime, bool bLoop,
+	bool bMirrored, UMirrorDataTable* MirrorDataTable, float BlendTime, const UBlendProfile* BlendProfile, EAlphaBlendOption BlendOption, bool bUseInertialBlend, FName InertialBlendNodeTag,
+	const FVector& BlendParameters, float PlayRate, float ActivationDelay, FName GroupName, EAnimGroupRole::Type GroupRole, EAnimSyncMethod GroupMethod, bool bOverridePositionWhenJoiningSyncGroupAsLeader,
+	FName PrevGroupName, EAnimGroupRole::Type PrevGroupRole, EAnimSyncMethod PrevMethod)
+{
+	if (!StitchTo(Context, AnimationAsset, AccumulatedTime, bLoop, bMirrored, MirrorDataTable, BlendTime, 
+		BlendProfile, BlendOption, bUseInertialBlend, InertialBlendNodeTag, BlendParameters, PlayRate, ActivationDelay))
 	{
 		InternalBlendTo(Context, AnimationAsset, AccumulatedTime, bLoop, bMirrored, MirrorDataTable,
-			BlendTime, BlendProfile, BlendOption, bUseInertialBlend, InertialBlendNodeTag, BlendParameters, PlayRate, ActivationDelay, GroupName, GroupRole, GroupMethod, bOverridePositionWhenJoiningSyncGroupAsLeader);
+			BlendTime, BlendProfile, BlendOption, bUseInertialBlend, InertialBlendNodeTag, BlendParameters, PlayRate, ActivationDelay, 
+			GroupName, GroupRole, GroupMethod, bOverridePositionWhenJoiningSyncGroupAsLeader, PrevGroupName, PrevGroupRole, PrevMethod);
 	}
 }
 
 void FAnimNode_BlendStack_Standalone::InternalBlendTo(const FAnimationUpdateContext& Context, UAnimationAsset* AnimationAsset, float AccumulatedTime, bool bLoop,
 	bool bMirrored, UMirrorDataTable* MirrorDataTable, float BlendTime, const UBlendProfile* BlendProfile, EAlphaBlendOption BlendOption, bool bUseInertialBlend, FName InertialBlendNodeTag,
-	const FVector& BlendParameters, float PlayRate, float ActivationDelay, FName GroupName, EAnimGroupRole::Type GroupRole, EAnimSyncMethod GroupMethod, bool bOverridePositionWhenJoiningSyncGroupAsLeader)
+	const FVector& BlendParameters, float PlayRate, float ActivationDelay, FName GroupName, EAnimGroupRole::Type GroupRole, EAnimSyncMethod GroupMethod, bool bOverridePositionWhenJoiningSyncGroupAsLeader,
+	FName PrevGroupName, EAnimGroupRole::Type PrevGroupRole, EAnimSyncMethod PrevGroupMethod)
 {
 	const bool bBlendStackIsEmpty = AnimPlayers.IsEmpty();
 
@@ -1107,9 +1188,7 @@ void FAnimNode_BlendStack_Standalone::InternalBlendTo(const FAnimationUpdateCont
 		RequestInertialBlend(Context, BlendTime, BlendProfile, BlendOption, InertialBlendNodeTag);
 		BlendTime = 0.0f;
 	}
-
 	
-
 	int32 NewSamplePoseLinkIndex = CurrentSamplePoseLink;
 	if (!bBlendStackIsEmpty && AnimPlayers[0].IsPostponed())
 	{
@@ -1157,15 +1236,42 @@ void FAnimNode_BlendStack_Standalone::InternalBlendTo(const FAnimationUpdateCont
 		UE_LOG(LogBlendStack, Warning, TEXT("FAnimNode_BlendStack_Standalone::InternalBlendTo multiple BlendTo requests during the same frame: only the last request will be put on this BlendStack"));
 	}
 
-	FBlendStackAnimPlayer& AnimPlayer = AnimPlayers[0];
+	if (AnimPlayers.Num() > 1)
+	{
+		FBlendStackAnimPlayer& PrevMainAnimPlayer = AnimPlayers[1];
+		if (FAnimNode_AssetPlayerBase* PrevAssetPlayerNode = PrevMainAnimPlayer.GetAssetPlayerNode())
+		{
+			PrevAssetPlayerNode->SetGroupMethod(PrevGroupMethod);
+			PrevAssetPlayerNode->SetGroupRole(PrevGroupRole);
+			PrevAssetPlayerNode->SetGroupName(PrevGroupName);
+		}
+
+		// Clear up any sync group of other AnimPlayers
+		for (int32 i = 2; i < AnimPlayers.Num(); ++i)
+		{
+			FBlendStackAnimPlayer& AnimPlayer = AnimPlayers[i];
+			if (FAnimNode_AssetPlayerBase* AssetPlayerNode = AnimPlayer.GetAssetPlayerNode())
+			{
+				AssetPlayerNode->SetGroupMethod(EAnimSyncMethod::DoNotSync);
+				AssetPlayerNode->SetGroupRole(EAnimGroupRole::CanBeLeader);
+				AssetPlayerNode->SetGroupName(NAME_None);
+			}
+		}
+	}
+
+	FBlendStackAnimPlayer& MainAnimPlayer = AnimPlayers[0];
 
 	FAnimationInitializeContext InitContext(Context.AnimInstanceProxy, Context.SharedContext);
-	AnimPlayer.Initialize(InitContext, AnimationAsset, AccumulatedTime, bLoop, bMirrored, MirrorDataTable, BlendTime, BlendProfile, BlendOption, BlendParameters, PlayRate, ActivationDelay, NewSamplePoseLinkIndex, GroupName, GroupRole, GroupMethod, bOverridePositionWhenJoiningSyncGroupAsLeader);
-	InitializeSample(InitContext, AnimPlayer);
+	MainAnimPlayer.Initialize(InitContext, AnimationAsset, AccumulatedTime, bLoop, bMirrored, MirrorDataTable, BlendTime, BlendProfile, BlendOption, BlendParameters, PlayRate, ActivationDelay, NewSamplePoseLinkIndex, GroupName, GroupRole, GroupMethod, bOverridePositionWhenJoiningSyncGroupAsLeader);
+	InitializeSample(InitContext, MainAnimPlayer);
+
+	LastBlendToFrameCounter = GFrameCounter;
 }
 
 void FAnimNode_BlendStack_Standalone::Reset()
 {
+	LastBlendToFrameCounter = GFrameCounter;
+
 	// reserving MaxActiveBlends + 2 AnimPlayers, to avoid any reallocation
 	AnimPlayers.Reserve(MaxActiveBlends + 2);
 	AnimPlayers.Reset();
@@ -1191,6 +1297,11 @@ int32 FAnimNode_BlendStack_Standalone::GetNextPoseLinkIndex()
 	return CurrentSamplePoseLink;
 }
 
+bool FAnimNode_BlendStack_Standalone::AnyNewBlendToThisFrame() const
+{
+	return LastBlendToFrameCounter == GFrameCounter;
+}
+
 void FAnimNode_BlendStack_Standalone::UpdatePlayRate(float PlayRate)
 {
 	if (!AnimPlayers.IsEmpty())
@@ -1382,6 +1493,20 @@ void FAnimNode_BlendStack::UpdateAssetPlayer(const FAnimationUpdateContext& Cont
 	UpdatePlayRate(WantedPlayRate);
 	UpdateBlendspaceParameters(BlendspaceUpdateMode, BlendParameters);
 
+	if (bForceBlendTo && !AnyNewBlendToThisFrame())
+	{
+		check(!AnimPlayers.IsEmpty());
+		const FBlendStackAnimPlayer& MainAnimPlayer = AnimPlayers[0];
+
+		// blending into the new main anim player with the new integrated time to avoid pops on this discontinuity
+		InternalBlendTo(Context, MainAnimPlayer.GetAnimationAsset(), MainAnimPlayer.GetAccumulatedTime(), MainAnimPlayer.IsLooping(), MainAnimPlayer.GetMirror(),
+			MainAnimPlayer.GetMirrorDataTable(), BlendTime,
+			MainAnimPlayer.GetBlendProfile(), MainAnimPlayer.GetBlendOption(), bUseInertialBlend, 
+			InertialBlendNodeTag, MainAnimPlayer.GetBlendParameters(), MainAnimPlayer.GetPlayRate(), 0.f,
+			GetGroupName(), GetGroupRole(), GetGroupMethod(), GetOverridePositionWhenJoiningSyncGroupAsLeader(),
+			GetPrevGroupName(), GetPrevGroupRole(), GetPrevGroupMethod());
+	}
+
 	Super::UpdateAssetPlayer(Context);
 }
 
@@ -1412,6 +1537,21 @@ EAnimSyncMethod FAnimNode_BlendStack::GetGroupMethod() const
 	return GET_ANIM_NODE_DATA(EAnimSyncMethod, Method);
 }
 
+FName FAnimNode_BlendStack::GetPrevGroupName() const
+{
+	return GET_ANIM_NODE_DATA(FName, PrevGroupName);
+}
+
+EAnimGroupRole::Type FAnimNode_BlendStack::GetPrevGroupRole() const
+{
+	return GET_ANIM_NODE_DATA(TEnumAsByte<EAnimGroupRole::Type>, PrevGroupRole);
+}
+
+EAnimSyncMethod FAnimNode_BlendStack::GetPrevGroupMethod() const
+{
+	return GET_ANIM_NODE_DATA(EAnimSyncMethod, PrevMethod);
+}
+
 bool FAnimNode_BlendStack::GetIgnoreForRelevancyTest() const
 {
 	return GET_ANIM_NODE_DATA(bool, bIgnoreForRelevancyTest);
diff --git a/Engine/Plugins/Animation/BlendStack/Source/Runtime/Public/BlendStack/AnimNode_BlendStack.h b/Engine/Plugins/Animation/BlendStack/Source/Runtime/Public/BlendStack/AnimNode_BlendStack.h
--- a/Engine/Plugins/Animation/BlendStack/Source/Runtime/Public/BlendStack/AnimNode_BlendStack.h
+++ b/Engine/Plugins/Animation/BlendStack/Source/Runtime/Public/BlendStack/AnimNode_BlendStack.h
@@ -29,7 +29,7 @@ struct FBlendStackAnimPlayer
 	friend struct FAnimNode_BlendStack_Standalone;
 	
 	UE_API void Initialize(const FAnimationInitializeContext& Context, UAnimationAsset* AnimationAsset, float AccumulatedTime, bool bLoop,
-		bool bMirrored, UMirrorDataTable* MirrorDataTable, float BlendTime, const UBlendProfile* BlendProfile, EAlphaBlendOption InBlendOption,
+		bool bMirrored, UMirrorDataTable* MirrorDataTable, float BlendTime, const UBlendProfile* InBlendProfile, EAlphaBlendOption InBlendOption,
 		const FVector& BlendParameters, float PlayRate, float ActivationDelay, int32 InPoseLinkIdx, FName GroupName, EAnimGroupRole::Type GroupRole, EAnimSyncMethod Method, bool bOverridePositionWhenJoiningSyncGroupAsLeader);
 	
 	UE_API void UpdatePlayRate(float PlayRate);
@@ -48,6 +48,13 @@ struct FBlendStackAnimPlayer
 	UE_API float GetBlendInWeight() const;
 
 	EAlphaBlendOption GetBlendOption() const { return BlendOption; }
+	UMirrorDataTable* GetMirrorDataTable() const;
+	EAnimSyncMethod GetGroupMethod() const;
+	FName GetGroupName() const;
+	EAnimGroupRole::Type GetGroupRole() const;
+	bool GetOverridePositionWhenJoiningSyncGroupAsLeader() const;
+	const UBlendProfile* GetBlendProfile() const;
+
 	UE_API void StorePoseContext(const FPoseContext& PoseContext);
 	UE_API bool HasValidPoseContext() const;
 	UE_API void MovePoseContextTo(FBlendStackAnimPlayer& Other);
@@ -113,6 +120,9 @@ private:
 	int32 PoseLinkIndex = INDEX_NONE;
 	TCustomBoneIndexArray<float, FSkeletonPoseBoneIndex> TotalBlendInTimePerBone;
 
+	UPROPERTY(Transient)
+	TObjectPtr<const UBlendProfile> BlendProfile = nullptr;
+
 	float TotalBlendInTime = 0.f;
 	float CurrentBlendInTime = 0.f;
 	float TimeToActivation = 0.f;
@@ -179,7 +189,8 @@ struct FAnimNode_BlendStack_Standalone : public FAnimNode_AssetPlayerBase
 		bool bMirrored = false, UMirrorDataTable* MirrorDataTable = nullptr, float BlendTime = 0.2f,
 		const UBlendProfile* BlendProfile = nullptr, EAlphaBlendOption BlendOption = EAlphaBlendOption::Linear, 
 		bool bUseInertialBlend = false, FName InertialBlendNodeTag = NAME_None, const FVector& BlendParameters = FVector::Zero(), float PlayRate = 1.f, float ActivationDelay = 0.f,
-		FName GroupName = NAME_None, EAnimGroupRole::Type GroupRole = EAnimGroupRole::CanBeLeader, EAnimSyncMethod Method = EAnimSyncMethod::DoNotSync, bool bOverridePositionWhenJoiningSyncGroupAsLeader = false);
+		FName GroupName = NAME_None, EAnimGroupRole::Type GroupRole = EAnimGroupRole::CanBeLeader, EAnimSyncMethod Method = EAnimSyncMethod::DoNotSync, bool bOverridePositionWhenJoiningSyncGroupAsLeader = false,
+		FName PrevGroupName = NAME_None, EAnimGroupRole::Type PrevGroupRole = EAnimGroupRole::CanBeLeader, EAnimSyncMethod PrevMethod = EAnimSyncMethod::DoNotSync);
 	UE_API void UpdatePlayRate(float PlayRate);
 	UE_API virtual void Reset();
 
@@ -199,12 +210,16 @@ struct FAnimNode_BlendStack_Standalone : public FAnimNode_AssetPlayerBase
 	int32 GetMaxActiveBlends() const { return MaxActiveBlends; }
 	void SetMaxActiveBlends(int32 InMaxActiveBlends) { MaxActiveBlends = InMaxActiveBlends; }
 	
+	// Experimental, this feature might be removed without warning, not for production use
+	UE_API bool AnyNewBlendToThisFrame() const;
+
 protected:
 	UE_API void InternalBlendTo(const FAnimationUpdateContext& Context, UAnimationAsset* AnimationAsset, float AccumulatedTime, bool bLoop, 
 		bool bMirrored, UMirrorDataTable* MirrorDataTable, float BlendTime,
 		const UBlendProfile* BlendProfile, EAlphaBlendOption BlendOption, 
 		bool bUseInertialBlend, FName InertialBlendNodeTag, const FVector& BlendParameters, float PlayRate, float ActivationDelay,
-		FName GroupName, EAnimGroupRole::Type GroupRole, EAnimSyncMethod Method, bool bOverridePositionWhenJoiningSyncGroupAsLeader);
+		FName GroupName, EAnimGroupRole::Type GroupRole, EAnimSyncMethod GroupMethod, bool bOverridePositionWhenJoiningSyncGroupAsLeader, 
+		FName PrevGroupName, EAnimGroupRole::Type PrevGroupRole, EAnimSyncMethod PrevGroupMethod);
 
 	// Call this to update current blendspace player xy blend parameters. By default, we only update them on initial BlendTo.
 	UE_API void UpdateBlendspaceParameters(const EBlendStack_BlendspaceUpdateMode UpdateMode, const FVector& BlendParameters);
@@ -239,12 +254,18 @@ protected:
 	float PlayerDepthBlendInTimeMultiplier = 1.f;
 
 private:
-	UE_API void PopLastAnimPlayer();
-	UE_API void InitializeSample(const FAnimationInitializeContext& Context, FBlendStackAnimPlayer& SamplePlayer);
-	UE_API bool EvaluateSample(FPoseContext& Output, const int32 PlayerIndex);
-	UE_API void UpdateSample(const FAnimationUpdateContext& Context, const int32 PlayerIndex);
-	UE_API void CacheBonesForSample(const FAnimationCacheBonesContext& Context, const int32 PlayerIndex);
-	UE_API bool IsSampleGraphAvailableForPlayer(const int32 PlayerIndex);
+	bool StitchTo(const FAnimationUpdateContext& Context, UAnimationAsset* AnimationAsset, float AccumulatedTime, bool bLoop,
+		bool bMirrored, UMirrorDataTable* MirrorDataTable, float BlendTime, const UBlendProfile* BlendProfile, EAlphaBlendOption BlendOption, bool bUseInertialBlend, FName InertialBlendNodeTag,
+		const FVector& BlendParameters, float PlayRate, float ActivationDelay);
+
+	void PopLastAnimPlayer();
+	void InitializeSample(const FAnimationInitializeContext& Context, FBlendStackAnimPlayer& SamplePlayer);
+	bool EvaluateSample(FPoseContext& Output, const int32 PlayerIndex);
+	void UpdateSample(const FAnimationUpdateContext& Context, const int32 PlayerIndex);
+	void CacheBonesForSample(const FAnimationCacheBonesContext& Context, const int32 PlayerIndex);
+	bool IsSampleGraphAvailableForPlayer(const int32 PlayerIndex);
+	
+	uint64 LastBlendToFrameCounter = 0;
 };
 
 USTRUCT(BlueprintInternalUseOnly)
@@ -309,6 +330,11 @@ struct FAnimNode_BlendStack : public FAnimNode_BlendStack_Standalone
 	UPROPERTY(EditAnywhere, Category = Blendspace)
 	float BlendParametersDeltaThreshold = 0.0f;
 
+	// in case there're no blend to the current frame, this blend stack will force a blend to into the currently playing animation. 
+	// useful to blend into the same animation in case this blendstack synch properties are set to be a follower
+	UPROPERTY(Experimental, EditAnywhere, BlueprintReadWrite, Category = "Experimental", meta = (PinHiddenByDefault))
+	bool bForceBlendTo = false;
+
 	// Enable to use inertial blending
 	UPROPERTY(EditAnywhere, Category = Settings)
 	bool bUseInertialBlend = false;
@@ -328,16 +354,26 @@ struct FAnimNode_BlendStack : public FAnimNode_BlendStack_Standalone
     UPROPERTY(EditAnywhere, Category = Sync, meta = (FoldProperty))
     FName GroupName = NAME_None;
     
-    // The role this node can assume within the group (ignored if GroupName is not set). Note
-    // that this is the role of the output of this node, not of animations contained by it.
+    // The role this node (and main AnimPlayer - most recent animation set via BlendTo) can assume within the group (ignored if GroupName is not set)
     UPROPERTY(EditAnywhere, Category = Sync, meta = (FoldProperty))
     TEnumAsByte<EAnimGroupRole::Type> GroupRole = EAnimGroupRole::CanBeLeader;
     
-    // How this node will synchronize with other animations. Note that this determines how the output
-    // of this node is used for synchronization, not of animations contained by it.
+    // How this node main AnimPlayer will synchronize with other animations
     UPROPERTY(EditAnywhere, Category = Sync, meta = (FoldProperty))
     EAnimSyncMethod Method = EAnimSyncMethod::DoNotSync;
     
+    // The role set to the most relevant non main AnimPlayer (previous main AnimPlayer)
+    UPROPERTY(EditAnywhere, Category = Sync, meta = (FoldProperty))
+    FName PrevGroupName = NAME_None;
+    
+	// The group set to the most relevant non main AnimPlayer (previous main AnimPlayer)
+    UPROPERTY(EditAnywhere, Category = Sync, meta = (FoldProperty))
+    TEnumAsByte<EAnimGroupRole::Type> PrevGroupRole = EAnimGroupRole::CanBeLeader;
+    
+	// The method set to the most relevant non main AnimPlayer (previous main AnimPlayer)
+    UPROPERTY(EditAnywhere, Category = Sync, meta = (FoldProperty))
+    EAnimSyncMethod PrevMethod = EAnimSyncMethod::DoNotSync;
+
     // If true, "Relevant anim" nodes that look for the highest weighted animation in a state will ignore this node
     UPROPERTY(EditAnywhere, Category = Relevancy, meta = (FoldProperty, PinHiddenByDefault))
     bool bIgnoreForRelevancyTest = false;
@@ -348,6 +384,11 @@ struct FAnimNode_BlendStack : public FAnimNode_BlendStack_Standalone
 	UE_API virtual FName GetGroupName() const override;
 	UE_API virtual EAnimGroupRole::Type GetGroupRole() const override;
 	UE_API virtual EAnimSyncMethod GetGroupMethod() const override;
+
+	UE_API FName GetPrevGroupName() const;
+	UE_API EAnimGroupRole::Type GetPrevGroupRole() const;
+	UE_API EAnimSyncMethod GetPrevGroupMethod() const;
+
 	UE_API virtual bool GetIgnoreForRelevancyTest() const override;
 	UE_API virtual bool IsLooping() const override;
 	UE_API virtual bool SetGroupName(FName InGroupName) override;
diff --git a/Engine/Plugins/Animation/PoseSearch/Source/Runtime/Private/AnimNode_MotionMatching.cpp b/Engine/Plugins/Animation/PoseSearch/Source/Runtime/Private/AnimNode_MotionMatching.cpp
--- a/Engine/Plugins/Animation/PoseSearch/Source/Runtime/Private/AnimNode_MotionMatching.cpp
+++ b/Engine/Plugins/Animation/PoseSearch/Source/Runtime/Private/AnimNode_MotionMatching.cpp
@@ -250,17 +250,6 @@ void FAnimNode_MotionMatching::UpdateAssetPlayer(const FAnimationUpdateContext&
 			{
 				if (UAnimationAsset* AnimationAsset = MotionMatchingState.SearchResult.GetAnimationAssetForRole())
 				{
-					// Clear up any sync group info before pushing new asset player (which will have sync info since its the highest weighted).
-					for (FBlendStackAnimPlayer& AnimPlayer : AnimPlayers)
-					{
-						if (FAnimNode_AssetPlayerBase* AssetPlayerNode = AnimPlayer.GetAssetPlayerNode())
-						{
-							AssetPlayerNode->SetGroupMethod(EAnimSyncMethod::DoNotSync);
-							AssetPlayerNode->SetGroupRole(EAnimGroupRole::CanBeLeader);
-							AssetPlayerNode->SetGroupName(NAME_None);
-						}
-					}
-
 					bool bAbortBlend = false;
 					if (!AnimPlayers.IsEmpty() && GetBlendspaceParametersDeltaThreshold() > 0.f && AnimationAsset->IsA<UBlendSpace>())
 					{
diff --git a/Engine/Plugins/Animation/PoseSearch/Source/Runtime/Private/MotionMatchingAnimNodeLibrary.cpp b/Engine/Plugins/Animation/PoseSearch/Source/Runtime/Private/MotionMatchingAnimNodeLibrary.cpp
--- a/Engine/Plugins/Animation/PoseSearch/Source/Runtime/Private/MotionMatchingAnimNodeLibrary.cpp
+++ b/Engine/Plugins/Animation/PoseSearch/Source/Runtime/Private/MotionMatchingAnimNodeLibrary.cpp
@@ -150,3 +150,14 @@ void UMotionMatchingAnimNodeLibrary::SetInterruptMode(const FMotionMatchingAnimN
 		UE_LOG(LogPoseSearch, Warning, TEXT("UMotionMatchingAnimNodeLibrary::SetInterruptMode called on an invalid context or with an invalid type"));
 	}
 }
+
+bool UMotionMatchingAnimNodeLibrary::AnyNewBlendToThisFrame(const FMotionMatchingAnimNodeReference& MotionMatchingNode)
+{
+	if (const FAnimNode_MotionMatching* MotionMatchingNodePtr = MotionMatchingNode.GetAnimNodePtr<FAnimNode_MotionMatching>())
+	{
+		return MotionMatchingNodePtr->AnyNewBlendToThisFrame();
+	}
+	
+	UE_LOG(LogPoseSearch, Warning, TEXT("UMotionMatchingAnimNodeLibrary::AnyNewBlendToThisFrame called on an invalid context or with an invalid type"));
+	return false;
+}
\ No newline at end of file
diff --git a/Engine/Plugins/Animation/PoseSearch/Source/Runtime/Public/PoseSearch/MotionMatchingAnimNodeLibrary.h b/Engine/Plugins/Animation/PoseSearch/Source/Runtime/Public/PoseSearch/MotionMatchingAnimNodeLibrary.h
--- a/Engine/Plugins/Animation/PoseSearch/Source/Runtime/Public/PoseSearch/MotionMatchingAnimNodeLibrary.h
+++ b/Engine/Plugins/Animation/PoseSearch/Source/Runtime/Public/PoseSearch/MotionMatchingAnimNodeLibrary.h
@@ -118,6 +118,13 @@ public:
 	 */
 	UFUNCTION(BlueprintCallable, Category = "Animation|MotionMatching", meta = (BlueprintThreadSafe))
 	static UE_API void SetInterruptMode(const FMotionMatchingAnimNodeReference& MotionMatchingNode, EPoseSearchInterruptMode InterruptMode);
+
+	/**
+	 * returns true if any new BlendTo API has been called so far during the current frame. BlendTo is the BlendStack API used by the motion matching node to blend into a newly selected animation
+	 * @param MotionMatchingNode - The motion matching node to operate on.
+	 */
+	UFUNCTION(Experimental, BlueprintPure, Category = "Animation|MotionMatching", meta = (BlueprintThreadSafe))
+	static UE_API bool AnyNewBlendToThisFrame(const FMotionMatchingAnimNodeReference& MotionMatchingNode);
 };
 
 #undef UE_API
