BlendStack - logging errors in FAnimNode_BlendStack_Standalone::InternalBlendTo if inconsistent an EAdditiveAnimationType gets detected

#Synced-CL 50154202

diff --git a/Engine/Plugins/Animation/BlendStack/Source/Runtime/Private/AnimNode_BlendStack.cpp b/Engine/Plugins/Animation/BlendStack/Source/Runtime/Private/AnimNode_BlendStack.cpp
--- a/Engine/Plugins/Animation/BlendStack/Source/Runtime/Private/AnimNode_BlendStack.cpp
+++ b/Engine/Plugins/Animation/BlendStack/Source/Runtime/Private/AnimNode_BlendStack.cpp
@@ -1182,6 +1182,45 @@ void FAnimNode_BlendStack_Standalone::InternalBlendTo(const FAnimationUpdateCont
 	{		
 		BlendTime = 0.0f;
 	}
+#if WITH_EDITOR
+	else
+	{
+		// returns AnimationAsset EAdditiveAnimationType. AAT_MAX if AnimationAsset is a UBlendSpace with inconsistent EAdditiveAnimationType
+		auto GetAdditiveAnimType = [](UAnimationAsset* AnimationAsset) -> EAdditiveAnimationType
+			{
+				EAdditiveAnimationType AdditiveAnimationType = AAT_None;
+				if (const UAnimSequenceBase* AnimSequenceBase = Cast<UAnimSequenceBase>(AnimationAsset))
+				{
+					AdditiveAnimationType = AnimSequenceBase->GetAdditiveAnimType();
+				}
+				else if (const UBlendSpace* BlendSpace = Cast<UBlendSpace>(AnimationAsset))
+				{
+					AdditiveAnimationType = AAT_MAX;
+					for (const FBlendSample& BlendSample : BlendSpace->GetBlendSamples())
+					{
+						if (AdditiveAnimationType == AAT_MAX)
+						{
+							AdditiveAnimationType = BlendSample.Animation->GetAdditiveAnimType();
+						}
+						else if (AdditiveAnimationType != BlendSample.Animation->GetAdditiveAnimType())
+						{
+							AdditiveAnimationType = AAT_MAX;
+							break;
+						}
+					}
+				}
+
+				return AdditiveAnimationType;
+			};
+
+		const EAdditiveAnimationType OldAdditiveAnimationType = GetAdditiveAnimType(AnimPlayers[0].GetAnimationAsset());
+		const EAdditiveAnimationType NewAdditiveAnimationType = GetAdditiveAnimType(AnimationAsset);
+		if (NewAdditiveAnimationType != OldAdditiveAnimationType)
+		{
+			UE_LOG(LogBlendStack, Error, TEXT("FAnimNode_BlendStack_Standalone::InternalBlendTo AdditiveAnimType mismatch between input animation '%s' and currently playing animation '%s'"), *GetNameSafe(AnimationAsset), *GetNameSafe(AnimPlayers[0].GetAnimationAsset()));
+		}
+	}
+#endif // WITH_EDITOR
 
 	if (bUseInertialBlend)
 	{
